Adding AutoLOD tools to the static mesh editor.

Static Mesh Editor changes:
- Share ownership of StaticMeshViewportClient's ModeTools via the EditorModeManager member
- Set up the EditorModeManager's preview scene and selected components when it's copied
- Add a set of FStaticMeshEditorToolbarExtender delegates to the module (similar to FSkeletalMeshEditorToolbarExtender)
- Allow adding an overlay widget to the viewport
- Add a ModeUILayer member to handle mode toolkit hosting

StaticMeshEditorModeling (new module):
- Create a mode and toolkit with AutoLOD and LODManager tools
- Add a button to the StaticMeshEditor's toolbar to toggle UMeshLODPluginEditMode

diff --git a/Engine/Plugins/Experimental/StaticMeshEditorModeling/Source/StaticMeshEditorModeling/Private/StaticMeshEditorModelingCommands.cpp b/Engine/Plugins/Experimental/StaticMeshEditorModeling/Source/StaticMeshEditorModeling/Private/StaticMeshEditorModelingCommands.cpp
new file mode 100644
--- /dev/null
+++ b/Engine/Plugins/Experimental/StaticMeshEditorModeling/Source/StaticMeshEditorModeling/Private/StaticMeshEditorModelingCommands.cpp
@@ -0,0 +1,23 @@
+// Copyright Epic Games, Inc. All Rights Reserved.
+
+#include "StaticMeshEditorModelingCommands.h"
+#include "EditorStyleSet.h"
+#include "Framework/Commands/UICommandInfo.h"
+
+#define LOCTEXT_NAMESPACE "StaticMeshEditorModelingCommands"
+
+FStaticMeshEditorModelingCommands::FStaticMeshEditorModelingCommands()
+	: TCommands<FStaticMeshEditorModelingCommands>(TEXT("MeshModelingTools"), NSLOCTEXT("Contexts", "MeshModelingTools", "Mesh Modeling"), NAME_None, FEditorStyle::GetStyleSetName())
+{}
+
+void FStaticMeshEditorModelingCommands::RegisterCommands()
+{
+	UI_COMMAND(ToggleStaticMeshEditorModelingMode, "Enable Mesh Modeling Tools", "Toggles Mesh Modeling Tools on or off.", EUserInterfaceActionType::ToggleButton, FInputChord());
+}
+
+const FStaticMeshEditorModelingCommands& FStaticMeshEditorModelingCommands::Get()
+{
+	return TCommands<FStaticMeshEditorModelingCommands>::Get();
+}
+
+#undef LOCTEXT_NAMESPACE
diff --git a/Engine/Plugins/Experimental/StaticMeshEditorModeling/Source/StaticMeshEditorModeling/Private/StaticMeshEditorModelingCommands.h b/Engine/Plugins/Experimental/StaticMeshEditorModeling/Source/StaticMeshEditorModeling/Private/StaticMeshEditorModelingCommands.h
new file mode 100644
--- /dev/null
+++ b/Engine/Plugins/Experimental/StaticMeshEditorModeling/Source/StaticMeshEditorModeling/Private/StaticMeshEditorModelingCommands.h
@@ -0,0 +1,16 @@
+// Copyright Epic Games, Inc. All Rights Reserved.
+
+#pragma once
+
+#include "Framework/Commands/Commands.h"
+
+class FStaticMeshEditorModelingCommands : public TCommands<FStaticMeshEditorModelingCommands>
+{
+public:
+	FStaticMeshEditorModelingCommands();
+	void RegisterCommands() override;
+	static const FStaticMeshEditorModelingCommands& Get();
+
+	TSharedPtr<FUICommandInfo> ToggleStaticMeshEditorModelingMode;
+};
+
diff --git a/Engine/Plugins/Experimental/StaticMeshEditorModeling/Source/StaticMeshEditorModeling/Private/StaticMeshEditorModelingMode.cpp b/Engine/Plugins/Experimental/StaticMeshEditorModeling/Source/StaticMeshEditorModeling/Private/StaticMeshEditorModelingMode.cpp
new file mode 100644
--- /dev/null
+++ b/Engine/Plugins/Experimental/StaticMeshEditorModeling/Source/StaticMeshEditorModeling/Private/StaticMeshEditorModelingMode.cpp
@@ -0,0 +1,43 @@
+// Copyright Epic Games, Inc. All Rights Reserved.
+
+#include "StaticMeshEditorModelingMode.h"
+#include "StaticMeshEditorModelingToolkit.h"
+#include "EdModeInteractiveToolsContext.h"
+#include "ToolTargetManager.h"
+#include "ModelingToolsManagerActions.h"
+#include "ToolTargets/StaticMeshComponentToolTarget.h"
+#include "Tools/GenerateStaticMeshLODAssetTool.h"
+#include "Tools/LODManagerTool.h"
+
+#define LOCTEXT_NAMESPACE "StaticMeshEditorModelingMode"
+
+const FEditorModeID UStaticMeshEditorModelingMode::Id("StaticMeshEditorModelingMode");
+
+UStaticMeshEditorModelingMode::UStaticMeshEditorModelingMode()
+{
+	Info = FEditorModeInfo(Id, LOCTEXT("StaticMeshEditorModelingMode", "StaticMeshEditorModelingMode"), FSlateIcon(), true);
+}
+
+void UStaticMeshEditorModelingMode::Enter()
+{
+	UEdMode::Enter();
+
+	GetInteractiveToolsContext()->TargetManager->AddTargetFactory(NewObject<UStaticMeshComponentToolTargetFactory>(GetToolManager()));
+
+	const FModelingToolsManagerCommands& ToolManagerCommands = FModelingToolsManagerCommands::Get();
+
+	RegisterTool(ToolManagerCommands.BeginGenerateStaticMeshLODAssetTool, TEXT("BeginGenerateStaticMeshLODAssetTool"), NewObject<UGenerateStaticMeshLODAssetToolBuilder>());
+	RegisterTool(ToolManagerCommands.BeginLODManagerTool, TEXT("BeginLODManagerTool"), NewObject<ULODManagerToolBuilder>());
+}
+
+bool UStaticMeshEditorModelingMode::UsesToolkits() const
+{ 
+	return true; 
+}
+
+void UStaticMeshEditorModelingMode::CreateToolkit()
+{
+	Toolkit = MakeShareable(new FStaticMeshEditorModelingToolkit);
+}
+
+#undef LOCTEXT_NAMESPACE
diff --git a/Engine/Plugins/Experimental/StaticMeshEditorModeling/Source/StaticMeshEditorModeling/Private/StaticMeshEditorModelingModule.cpp b/Engine/Plugins/Experimental/StaticMeshEditorModeling/Source/StaticMeshEditorModeling/Private/StaticMeshEditorModelingModule.cpp
new file mode 100644
--- /dev/null
+++ b/Engine/Plugins/Experimental/StaticMeshEditorModeling/Source/StaticMeshEditorModeling/Private/StaticMeshEditorModelingModule.cpp
@@ -0,0 +1,94 @@
+// Copyright Epic Games, Inc. All Rights Reserved.
+
+#include "StaticMeshEditorModelingModule.h"
+#include "StaticMeshEditorModelingCommands.h"
+#include "StaticMeshEditorModelingMode.h"
+#include "ToolMenus.h"
+#include "StaticMeshEditorModule.h"
+#include "IStaticMeshEditor.h"
+#include "EditorModeManager.h"
+
+#define LOCTEXT_NAMESPACE "FStaticMeshEditorModelingModule"
+
+void FStaticMeshEditorModelingModule::StartupModule()
+{
+	// This code will execute after your module is loaded into memory; the exact timing is specified in the .uplugin file per-module
+	FStaticMeshEditorModelingCommands::Register();
+
+	UToolMenus::RegisterStartupCallback(
+		FSimpleMulticastDelegate::FDelegate::CreateRaw(this, &FStaticMeshEditorModelingModule::RegisterMenusAndToolbars));
+
+	IStaticMeshEditorModule& StaticMeshEditorModule = FModuleManager::LoadModuleChecked<IStaticMeshEditorModule>("StaticMeshEditor");
+	TArray<IStaticMeshEditorModule::FStaticMeshEditorToolbarExtender>& ToolbarExtenders = StaticMeshEditorModule.GetAllStaticMeshEditorToolbarExtenders();
+
+	// Add a button to activate modeling mode to the static mesh editor toolbar
+	auto ExtendStaticMeshEditorToolbar = [this](const TSharedRef<FUICommandList> InCommandList, TSharedRef<IStaticMeshEditor> Editor) -> TSharedRef<FExtender>
+	{
+		TSharedPtr<FExtender> ToolbarExtender = MakeShareable(new FExtender);
+		TWeakPtr<IStaticMeshEditor> Ptr(Editor);
+
+		InCommandList->MapAction(FStaticMeshEditorModelingCommands::Get().ToggleStaticMeshEditorModelingMode,
+			FExecuteAction::CreateRaw(this, &FStaticMeshEditorModelingModule::OnToggleStaticMeshEditorModelingMode, Ptr),
+			FCanExecuteAction(),
+			FIsActionChecked::CreateRaw(this, &FStaticMeshEditorModelingModule::IsStaticMeshEditorModelingModeActive, Ptr));
+
+		return ToolbarExtender.ToSharedRef();
+	};
+
+	ToolbarExtenders.Add(IStaticMeshEditorModule::FStaticMeshEditorToolbarExtender::CreateLambda(ExtendStaticMeshEditorToolbar));
+}
+
+
+void FStaticMeshEditorModelingModule::ShutdownModule()
+{
+	UToolMenus::UnregisterOwner(this);
+	FStaticMeshEditorModelingCommands::Unregister();
+}
+
+bool FStaticMeshEditorModelingModule::IsStaticMeshEditorModelingModeActive(TWeakPtr<IStaticMeshEditor> InEditor) const
+{
+	TSharedPtr<IStaticMeshEditor> PinnedEditor = InEditor.Pin();
+	return PinnedEditor.IsValid() && PinnedEditor->GetEditorModeManager().IsModeActive(UStaticMeshEditorModelingMode::Id);
+}
+
+
+void FStaticMeshEditorModelingModule::OnToggleStaticMeshEditorModelingMode(TWeakPtr<IStaticMeshEditor> InEditor)
+{
+	TSharedPtr<IStaticMeshEditor> PinnedEditor = InEditor.Pin();
+	if (PinnedEditor.IsValid())
+	{
+		if (!IsStaticMeshEditorModelingModeActive(InEditor))
+		{
+			PinnedEditor->GetEditorModeManager().ActivateMode(UStaticMeshEditorModelingMode::Id, true);
+		}
+		else
+		{
+			PinnedEditor->GetEditorModeManager().ActivateDefaultMode();
+		}
+	}
+}
+
+void FStaticMeshEditorModelingModule::RegisterMenusAndToolbars()
+{
+	FToolMenuOwnerScoped OwnerScoped(this);
+
+	UToolMenu* Toolbar = UToolMenus::Get()->ExtendMenu("AssetEditor.StaticMeshEditor.ToolBar");
+
+	{
+		FToolMenuSection& Section = Toolbar->FindOrAddSection("StaticMesh");
+		Section.AddDynamicEntry("ToggleStaticMeshEditorModelingMode", FNewToolMenuSectionDelegate::CreateLambda([this](FToolMenuSection& InSection)
+		{
+			InSection.AddEntry(FToolMenuEntry::InitToolBarButton(
+				FStaticMeshEditorModelingCommands::Get().ToggleStaticMeshEditorModelingMode,
+				LOCTEXT("StaticMeshEditorModelingMode", "Mesh Modeling Tools"),
+				FText(),
+				FSlateIcon()));
+		}));
+	}
+}
+
+
+#undef LOCTEXT_NAMESPACE
+	
+IMPLEMENT_MODULE(FStaticMeshEditorModelingModule, StaticMeshEditorModeling)
+
diff --git a/Engine/Plugins/Experimental/StaticMeshEditorModeling/Source/StaticMeshEditorModeling/Private/StaticMeshEditorModelingToolkit.cpp b/Engine/Plugins/Experimental/StaticMeshEditorModeling/Source/StaticMeshEditorModeling/Private/StaticMeshEditorModelingToolkit.cpp
new file mode 100644
--- /dev/null
+++ b/Engine/Plugins/Experimental/StaticMeshEditorModeling/Source/StaticMeshEditorModeling/Private/StaticMeshEditorModelingToolkit.cpp
@@ -0,0 +1,286 @@
+// Copyright Epic Games, Inc. All Rights Reserved.
+
+#include "StaticMeshEditorModelingToolkit.h"
+#include "Widgets/Images/SImage.h"
+#include "Widgets/Input/SButton.h"
+#include "ModelingToolsEditorModeStyle.h"
+#include "ModelingToolsManagerActions.h"
+#include "Tools/UEdMode.h"
+#include "InteractiveTool.h"
+#include "EdModeInteractiveToolsContext.h"
+#include "InteractiveToolManager.h"
+#include "Widgets/Text/STextBlock.h"
+#include "EditorStyleSet.h"
+#include "IDetailsView.h"
+#include "Framework/MultiBox/MultiBoxBuilder.h"
+
+
+#define LOCTEXT_NAMESPACE "StaticMeshEditorModelingToolkit"
+
+static const FName LODManagementTabName(TEXT("LODManagement"));
+
+void FStaticMeshEditorModelingToolkit::Init(
+	const TSharedPtr<IToolkitHost>& InToolkitHost,
+	TWeakObjectPtr<UEdMode> InOwningMode )
+{
+	ModeWarningArea = SNew(STextBlock)
+		.AutoWrapText(true)
+		.Font(FCoreStyle::GetDefaultFontStyle("Bold", 9))
+		.ColorAndOpacity(FSlateColor(FLinearColor(0.9f, 0.15f, 0.15f)));
+	ModeWarningArea->SetText(FText::GetEmpty());
+	ModeWarningArea->SetVisibility(EVisibility::Collapsed);
+
+	ModeHeaderArea = SNew(STextBlock)
+		.AutoWrapText(true)
+		.Font(FCoreStyle::GetDefaultFontStyle("Bold", 12));
+	ModeHeaderArea->SetText(LOCTEXT("SelectToolLabel", "Select a Tool from the Toolbar"));
+	ModeHeaderArea->SetJustification(ETextJustify::Center);
+
+	ToolWarningArea = SNew(STextBlock)
+		.AutoWrapText(true)
+		.Font(FCoreStyle::GetDefaultFontStyle("Bold", 9))
+		.ColorAndOpacity(FSlateColor(FLinearColor(0.9f, 0.15f, 0.15f)));
+	ToolWarningArea->SetText(FText::GetEmpty());
+
+	SAssignNew(ToolkitWidget, SBorder)
+		.HAlign(HAlign_Fill)
+		.Padding(4)
+		[
+			SNew(SVerticalBox)
+			+ SVerticalBox::Slot().AutoHeight().HAlign(HAlign_Fill).Padding(5)
+			[
+				ModeWarningArea->AsShared()
+			]
+			+ SVerticalBox::Slot().AutoHeight().HAlign(HAlign_Fill).Padding(5)
+			[
+				ModeHeaderArea->AsShared()
+			]
+			+ SVerticalBox::Slot().AutoHeight().HAlign(HAlign_Fill).Padding(5)
+			[
+				ToolWarningArea->AsShared()
+			]
+			+ SVerticalBox::Slot().FillHeight(1.f).HAlign(HAlign_Fill)
+			[
+				SAssignNew(ToolDetailsContainer, SBorder)
+				.BorderImage(FEditorStyle::GetBrush("NoBorder"))
+			]
+		];
+
+
+	FModeToolkit::Init(InToolkitHost, InOwningMode);
+
+	ToolDetailsContainer->SetContent(DetailsView.ToSharedRef());
+
+	ClearNotification();
+	ClearWarning();
+
+	ActiveToolName = FText::GetEmpty();
+	ActiveToolMessage = FText::GetEmpty();
+
+	GetScriptableEditorMode()->GetInteractiveToolsContext()->OnToolNotificationMessage.AddSP(this, &FStaticMeshEditorModelingToolkit::PostNotification);
+	GetScriptableEditorMode()->GetInteractiveToolsContext()->OnToolWarningMessage.AddSP(this, &FStaticMeshEditorModelingToolkit::PostWarning);
+
+
+	// ViewportOverlay
+	SAssignNew(ViewportOverlayWidget, SHorizontalBox)
+
+	+SHorizontalBox::Slot()
+	.HAlign(HAlign_Center)
+	.VAlign(VAlign_Bottom)
+	.Padding(FMargin(0.0f, 0.0f, 0.f, 15.f))
+	[
+		SNew(SBorder)
+		.BorderImage(FAppStyle::Get().GetBrush("EditorViewport.OverlayBrush"))
+		.Padding(8.f)
+		[
+			SNew(SHorizontalBox)
+
+			//+SHorizontalBox::Slot()
+			//.AutoWidth()
+			//.VAlign(VAlign_Center)
+			//.Padding(FMargin(0.f, 0.f, 8.f, 0.f))
+			//[
+			//	SNew(SImage)
+			//	.Image_Lambda([this] () { return ActiveToolIcon; })
+			//]
+
+			+SHorizontalBox::Slot()
+			.AutoWidth()
+			.VAlign(VAlign_Center)
+			.Padding(FMargin(0.f, 0.f, 8.f, 0.f))
+			[
+				SNew(STextBlock)
+				.Text(this, &FStaticMeshEditorModelingToolkit::GetActiveToolDisplayName)
+			]
+
+			+SHorizontalBox::Slot()
+			.AutoWidth()
+			.Padding(FMargin(0.0, 0.f, 2.f, 0.f))
+			[
+				SNew(SButton)
+				.ButtonStyle(FAppStyle::Get(), "PrimaryButton")
+				.TextStyle( FAppStyle::Get(), "DialogButtonText" )
+				.Text(LOCTEXT("OverlayAccept", "Accept"))
+				.ToolTipText(LOCTEXT("OverlayAcceptTooltip", "Accept/Commit the results of the active Tool [Enter]"))
+				.HAlign(HAlign_Center)
+				.OnClicked_Lambda([this]() { GetScriptableEditorMode()->GetInteractiveToolsContext()->EndTool(EToolShutdownType::Accept); return FReply::Handled(); })
+				.IsEnabled_Lambda([this]() { return GetScriptableEditorMode()->GetInteractiveToolsContext()->CanAcceptActiveTool(); })
+				.Visibility_Lambda([this]() { return GetScriptableEditorMode()->GetInteractiveToolsContext()->ActiveToolHasAccept() ? EVisibility::Visible : EVisibility::Collapsed; })
+			]
+
+			+SHorizontalBox::Slot()
+			.AutoWidth()
+			.Padding(FMargin(2.0, 0.f, 0.f, 0.f))
+			[
+				SNew(SButton)
+				.ButtonStyle(FAppStyle::Get(), "Button")
+				.TextStyle( FAppStyle::Get(), "DialogButtonText" )
+				.Text(LOCTEXT("OverlayCancel", "Cancel"))
+				.ToolTipText(LOCTEXT("OverlayCancelTooltip", "Cancel the active Tool [Esc]"))
+				.HAlign(HAlign_Center)
+				.OnClicked_Lambda([this]() { GetScriptableEditorMode()->GetInteractiveToolsContext()->EndTool(EToolShutdownType::Cancel); return FReply::Handled(); })
+				.IsEnabled_Lambda([this]() { return GetScriptableEditorMode()->GetInteractiveToolsContext()->CanCancelActiveTool(); })
+				.Visibility_Lambda([this]() { return GetScriptableEditorMode()->GetInteractiveToolsContext()->ActiveToolHasAccept() ? EVisibility::Visible : EVisibility::Collapsed; })
+			]
+
+			//+SHorizontalBox::Slot()
+			//.AutoWidth()
+			//.Padding(FMargin(2.0, 0.f, 0.f, 0.f))
+			//[
+			//	SNew(SButton)
+			//	.ButtonStyle(FAppStyle::Get(), "PrimaryButton")
+			//	.TextStyle( FAppStyle::Get(), "DialogButtonText" )
+			//	.Text(LOCTEXT("OverlayComplete", "Complete"))
+			//	.ToolTipText(LOCTEXT("OverlayCompleteTooltip", "Exit the active Tool [Enter]"))
+			//	.HAlign(HAlign_Center)
+			//	.OnClicked_Lambda([this]() { GetScriptableEditorMode()->GetInteractiveToolsContext()->EndTool(EToolShutdownType::Completed); return FReply::Handled(); })
+			//	.IsEnabled_Lambda([this]() { return GetScriptableEditorMode()->GetInteractiveToolsContext()->CanCompleteActiveTool(); })
+			//	.Visibility_Lambda([this]() { return GetScriptableEditorMode()->GetInteractiveToolsContext()->CanCompleteActiveTool() ? EVisibility::Visible : EVisibility::Collapsed; })
+			//]
+		]	
+	];
+
+	
+	SetCurrentPalette("LODManagement");
+}
+
+
+TSharedPtr<SWidget> FStaticMeshEditorModelingToolkit::GetInlineContent() const
+{
+	return ToolkitWidget;
+}
+
+FName FStaticMeshEditorModelingToolkit::GetToolkitFName() const
+{
+	return FName("MeshLODEditModeToolkit");
+}
+
+FText FStaticMeshEditorModelingToolkit::GetBaseToolkitName() const
+{
+	return LOCTEXT("DisplayName", "FMeshLODPluginEditModeToolkit Tool");
+}
+
+
+void FStaticMeshEditorModelingToolkit::GetToolPaletteNames(TArray<FName>& PaletteNames) const
+{
+	PaletteNames = { LODManagementTabName };
+}
+
+FText FStaticMeshEditorModelingToolkit::GetToolPaletteDisplayName(FName PaletteName) const
+{
+	return FText::FromName(PaletteName);
+}
+
+void FStaticMeshEditorModelingToolkit::BuildToolPalette(FName PaletteName, class FToolBarBuilder& ToolbarBuilder)
+{
+	const FModelingToolsManagerCommands& Commands = FModelingToolsManagerCommands::Get();
+	
+	if (PaletteName == LODManagementTabName)
+	{
+		ToolbarBuilder.AddToolBarButton(Commands.BeginGenerateStaticMeshLODAssetTool);
+		ToolbarBuilder.AddToolBarButton(Commands.BeginLODManagerTool);
+	}
+}
+
+void FStaticMeshEditorModelingToolkit::PostNotification(const FText& InMessage)
+{
+	ActiveToolMessage = InMessage;
+}
+
+
+void FStaticMeshEditorModelingToolkit::ClearNotification()
+{
+	ActiveToolMessage = FText::GetEmpty();
+}
+
+FText FStaticMeshEditorModelingToolkit::GetActiveToolDisplayName() const
+{
+	return ActiveToolName;
+}
+
+FText FStaticMeshEditorModelingToolkit::GetActiveToolMessage() const
+{
+	return ActiveToolMessage;
+}
+
+void FStaticMeshEditorModelingToolkit::PostWarning(const FText& Message)
+{
+	ToolWarningArea->SetText(Message);
+	ToolWarningArea->SetVisibility(EVisibility::Visible);
+}
+
+
+void FStaticMeshEditorModelingToolkit::ClearWarning()
+{
+	ToolWarningArea->SetText(FText());
+	ToolWarningArea->SetVisibility(EVisibility::Collapsed);
+}
+
+void FStaticMeshEditorModelingToolkit::UpdateActiveToolProperties(UInteractiveTool* Tool)
+{
+	if (Tool)
+	{
+		DetailsView->SetObjects(Tool->GetToolProperties(true));
+	}
+}
+
+void FStaticMeshEditorModelingToolkit::OnToolStarted(UInteractiveToolManager* Manager, UInteractiveTool* Tool)
+{
+	FModeToolkit::OnToolStarted(Manager, Tool);
+
+	Tool->OnPropertySetsModified.AddSP(this, &FStaticMeshEditorModelingToolkit::UpdateActiveToolProperties, Tool);
+
+	ModeHeaderArea->SetVisibility(EVisibility::Collapsed);
+	ActiveToolName = Tool->GetToolInfo().ToolDisplayName;
+
+	FString ActiveToolIdentifier = GetScriptableEditorMode()->GetToolManager()->GetActiveToolName(EToolSide::Left);
+	ActiveToolIdentifier.InsertAt(0, ".");
+	FName ActiveToolIconName = ISlateStyle::Join(FModelingToolsManagerCommands::Get().GetContextName(), TCHAR_TO_ANSI(*ActiveToolIdentifier));
+	ActiveToolIcon = FModelingToolsEditorModeStyle::Get()->GetOptionalBrush(ActiveToolIconName);
+
+	GetToolkitHost()->AddViewportOverlayWidget(ViewportOverlayWidget.ToSharedRef());
+}
+
+
+void FStaticMeshEditorModelingToolkit::OnToolEnded(UInteractiveToolManager* Manager, UInteractiveTool* Tool)
+{
+	FModeToolkit::OnToolEnded(Manager, Tool);
+
+	if (IsHosted())
+	{
+		GetToolkitHost()->RemoveViewportOverlayWidget(ViewportOverlayWidget.ToSharedRef());
+	}
+
+	if (Tool)
+	{
+		Tool->OnPropertySetsModified.RemoveAll(this);
+	}
+
+	ModeHeaderArea->SetVisibility(EVisibility::Visible);
+	ActiveToolName = FText::GetEmpty();
+	ActiveToolIcon = nullptr;
+	ClearNotification();
+	ClearWarning();
+}
+
+#undef LOCTEXT_NAMESPACE
diff --git a/Engine/Plugins/Experimental/StaticMeshEditorModeling/Source/StaticMeshEditorModeling/Public/StaticMeshEditorModelingMode.h b/Engine/Plugins/Experimental/StaticMeshEditorModeling/Source/StaticMeshEditorModeling/Public/StaticMeshEditorModelingMode.h
new file mode 100644
--- /dev/null
+++ b/Engine/Plugins/Experimental/StaticMeshEditorModeling/Source/StaticMeshEditorModeling/Public/StaticMeshEditorModelingMode.h
@@ -0,0 +1,30 @@
+// Copyright Epic Games, Inc. All Rights Reserved.
+
+#pragma once
+
+#include "Tools/LegacyEdModeWidgetHelpers.h"
+#include "Tools/Modes.h"
+#include "StaticMeshEditorModelingMode.generated.h"
+
+class SBorder;
+
+UCLASS()
+class UStaticMeshEditorModelingMode : public UBaseLegacyWidgetEdMode
+{
+	GENERATED_BODY()
+
+public:
+
+	const static FEditorModeID Id;
+
+protected:
+
+	UStaticMeshEditorModelingMode();
+
+	void Enter() override;
+
+	bool UsesToolkits() const override; 
+	void CreateToolkit() override;
+	
+
+};
diff --git a/Engine/Plugins/Experimental/StaticMeshEditorModeling/Source/StaticMeshEditorModeling/Public/StaticMeshEditorModelingModule.h b/Engine/Plugins/Experimental/StaticMeshEditorModeling/Source/StaticMeshEditorModeling/Public/StaticMeshEditorModelingModule.h
new file mode 100644
--- /dev/null
+++ b/Engine/Plugins/Experimental/StaticMeshEditorModeling/Source/StaticMeshEditorModeling/Public/StaticMeshEditorModelingModule.h
@@ -0,0 +1,29 @@
+// Copyright Epic Games, Inc. All Rights Reserved.
+
+#pragma once
+
+#include "Modules/ModuleInterface.h"
+#include "Templates/SharedPointer.h"
+
+class IStaticMeshEditor;
+class SDockTab;
+
+class FStaticMeshEditorModelingModule : public IModuleInterface
+{
+public:
+
+	/** IModuleInterface implementation */
+	virtual void StartupModule() override;
+	virtual void ShutdownModule() override;
+
+	bool IsStaticMeshEditorModelingModeActive(TWeakPtr<IStaticMeshEditor> InEditor) const;
+	void OnToggleStaticMeshEditorModelingMode(TWeakPtr<IStaticMeshEditor> InEditor);
+
+	void RegisterMenusAndToolbars();
+
+private:
+
+	TWeakPtr<SDockTab> MeshLODDockTab;
+
+};
+
diff --git a/Engine/Plugins/Experimental/StaticMeshEditorModeling/Source/StaticMeshEditorModeling/Public/StaticMeshEditorModelingToolkit.h b/Engine/Plugins/Experimental/StaticMeshEditorModeling/Source/StaticMeshEditorModeling/Public/StaticMeshEditorModelingToolkit.h
new file mode 100644
--- /dev/null
+++ b/Engine/Plugins/Experimental/StaticMeshEditorModeling/Source/StaticMeshEditorModeling/Public/StaticMeshEditorModelingToolkit.h
@@ -0,0 +1,52 @@
+// Copyright Epic Games, Inc. All Rights Reserved.
+
+#pragma once
+
+#include "Toolkits/BaseToolkit.h"
+
+class STextBlock;
+
+class FStaticMeshEditorModelingToolkit : public FModeToolkit
+{
+public:
+
+	void Init( const TSharedPtr<IToolkitHost>& InToolkitHost, TWeakObjectPtr<UEdMode> InOwningMode ) override;
+
+	FName GetToolkitFName() const override;
+	FText GetBaseToolkitName() const override;
+	TSharedPtr<SWidget> GetInlineContent() const override;
+
+	void GetToolPaletteNames(TArray<FName>& InPaletteName) const override;
+	FText GetToolPaletteDisplayName(FName PaletteName) const override;
+	void BuildToolPalette(FName PaletteName, class FToolBarBuilder& ToolbarBuilder) override;
+
+	void OnToolStarted(UInteractiveToolManager* Manager, UInteractiveTool* Tool) override;
+	void OnToolEnded(UInteractiveToolManager* Manager, UInteractiveTool* Tool) override;
+
+private:
+
+	void PostNotification(const FText& Message);
+	void ClearNotification();
+	void PostWarning(const FText& Message);
+	void ClearWarning();
+
+	void UpdateActiveToolProperties(UInteractiveTool* Tool);
+
+	FText GetActiveToolDisplayName() const override;
+	FText GetActiveToolMessage() const override;
+
+	TSharedPtr<SBorder> ToolDetailsContainer;
+
+	TSharedPtr<SWidget> ViewportOverlayWidget;
+	const FSlateBrush* ActiveToolIcon = nullptr;
+	FText ActiveToolName;
+	FText ActiveToolMessage;
+
+	TSharedPtr<STextBlock> ModeWarningArea;
+	TSharedPtr<STextBlock> ModeHeaderArea;
+	TSharedPtr<STextBlock> ToolWarningArea;
+
+};
+
+
+
diff --git a/Engine/Plugins/Experimental/StaticMeshEditorModeling/Source/StaticMeshEditorModeling/StaticMeshEditorModeling.Build.cs b/Engine/Plugins/Experimental/StaticMeshEditorModeling/Source/StaticMeshEditorModeling/StaticMeshEditorModeling.Build.cs
new file mode 100644
--- /dev/null
+++ b/Engine/Plugins/Experimental/StaticMeshEditorModeling/Source/StaticMeshEditorModeling/StaticMeshEditorModeling.Build.cs
@@ -0,0 +1,62 @@
+// Copyright Epic Games, Inc. All Rights Reserved.
+
+using UnrealBuildTool;
+
+public class StaticMeshEditorModeling : ModuleRules
+{
+	public StaticMeshEditorModeling(ReadOnlyTargetRules Target) : base(Target)
+	{
+		PCHUsage = ModuleRules.PCHUsageMode.UseExplicitOrSharedPCHs;
+		
+		PublicIncludePaths.AddRange(
+			new string[] {
+				// ... add public include paths required here ...
+			}
+			);
+				
+		
+		PrivateIncludePaths.AddRange(
+			new string[] {
+				// ... add other private include paths required here ...
+			}
+			);
+			
+		
+		PublicDependencyModuleNames.AddRange(
+			new string[]
+			{			
+				// ... add other public dependencies that you statically link with here ...
+			}
+			);
+			
+		
+		PrivateDependencyModuleNames.AddRange(
+			new string[]
+			{
+				"Core",
+				"CoreUObject",
+				"EditorInteractiveToolsFramework",
+				"EditorStyle",
+				"EditorFramework",
+				"InteractiveToolsFramework",
+				"MeshLODToolset",
+				"ModelingComponentsEditorOnly",
+				"ModelingToolsEditorMode",
+				"Slate",
+				"SlateCore",
+				"StaticMeshEditor",
+				"ToolMenus",
+				"UnrealEd",
+			}
+			);
+		
+		
+		DynamicallyLoadedModuleNames.AddRange(
+			new string[]
+			{
+				// ... add any modules that your module loads dynamically here ...
+			}
+			);
+
+	}
+}
diff --git a/Engine/Plugins/Experimental/StaticMeshEditorModeling/StaticMeshEditorModeling.uplugin b/Engine/Plugins/Experimental/StaticMeshEditorModeling/StaticMeshEditorModeling.uplugin
new file mode 100644
--- /dev/null
+++ b/Engine/Plugins/Experimental/StaticMeshEditorModeling/StaticMeshEditorModeling.uplugin
@@ -0,0 +1,35 @@
+{
+	"FileVersion": 3,
+	"Version": 1,
+	"VersionName": "0.1",
+	"FriendlyName": "Mesh LOD Plugin",
+	"Description": "A set of modules implementing 3D mesh LOD creation",
+	"Category": "Other",
+	"CreatedBy": "Epic Games, Inc.",
+	"CreatedByURL": "http://epicgames.com",
+	"DocsURL": "",
+	"MarketplaceURL": "",
+	"SupportURL": "",
+	"CanContainContent": true,
+	"IsBetaVersion": true,
+	"Installed": false,
+	"Hidden": true,
+	"Modules": [
+		{
+			"Name": "StaticMeshEditorModeling",
+			"Type": "Editor",
+			"LoadingPhase": "Default",
+			"WhitelistPlatforms": [ "Win64" ]
+		},
+	],
+	"Plugins": [
+		{
+			"Name": "MeshLODToolset",
+			"Enabled": true
+		},
+		{
+			"Name": "ModelingToolsEditorMode",
+			"Enabled": true
+		}
+	]
+}
diff --git a/Engine/Source/Editor/StaticMeshEditor/Private/SStaticMeshEditorViewport.h b/Engine/Source/Editor/StaticMeshEditor/Private/SStaticMeshEditorViewport.h
--- a/Engine/Source/Editor/StaticMeshEditor/Private/SStaticMeshEditorViewport.h
+++ b/Engine/Source/Editor/StaticMeshEditor/Private/SStaticMeshEditorViewport.h
@@ -127,6 +127,9 @@ public:
 
 	/** Returns the preview scene being renderd in the viewport */
 	TSharedRef<FAdvancedPreviewScene> GetPreviewScene() { return PreviewScene.ToSharedRef(); }
+
+	TSharedPtr<SOverlay> GetViewportOverlay() { return ViewportOverlay; }
+
 protected:
 	/** SEditorViewport interface */
 	virtual TSharedRef<FEditorViewportClient> MakeEditorViewportClient() override;
diff --git a/Engine/Source/Editor/StaticMeshEditor/Private/StaticMeshEditor.cpp b/Engine/Source/Editor/StaticMeshEditor/Private/StaticMeshEditor.cpp
--- a/Engine/Source/Editor/StaticMeshEditor/Private/StaticMeshEditor.cpp
+++ b/Engine/Source/Editor/StaticMeshEditor/Private/StaticMeshEditor.cpp
@@ -52,6 +52,11 @@
 #include "EditorViewportTabContent.h"
 #include "EditorViewportLayout.h"
 #include "Toolkits/AssetEditorToolkitMenuContext.h"
+#include "EditorModeManager.h"
+#include "StaticMeshEditorModeUILayer.h"
+#include "AssetEditorModeManager.h"
+#include "Engine/Selection.h"
+
 
 #define LOCTEXT_NAMESPACE "StaticMeshEditor"
 
@@ -261,12 +266,27 @@ void FStaticMeshEditor::InitStaticMeshEditor( const EToolkitMode::Type Mode, con
 	const bool bCreateDefaultToolbar = true;
 	FAssetEditorToolkit::InitAssetEditor( Mode, InitToolkitHost, StaticMeshEditorAppIdentifier, StandaloneDefaultLayout, bCreateDefaultToolbar, bCreateDefaultStandaloneMenu, ObjectToEdit );
 
+	
+	TSharedPtr<class IToolkitHost> PinnedToolkitHost = ToolkitHost.Pin();
+	check(PinnedToolkitHost.IsValid());
+	ModeUILayer = MakeShareable(new FStaticMeshEditorModeUILayer(PinnedToolkitHost.Get()));
+
+
 	ExtendMenu();
 	ExtendToolBar();
 	RegenerateMenusAndToolbars();
 	GenerateSecondaryToolbar();
 }
 
+
+void FStaticMeshEditor::PostInitAssetEditor()
+{
+	// (Copied from FUVEditorToolkit::PostInitAssetEditor)
+	// We need the viewport client to start out focused, or else it won't get ticked until we click inside it.
+	FStaticMeshEditorViewportClient& ViewportClient = GetStaticMeshViewport()->GetViewportClient();
+	ViewportClient.ReceivedFocus(ViewportClient.Viewport);
+}
+
 void FStaticMeshEditor::GenerateSecondaryToolbar()
 {
 	// Generate the secondary toolbar only if there are registered extensions
@@ -926,6 +946,16 @@ void FStaticMeshEditor::ExtendToolBar()
 	AddToolbarExtender(ToolbarExtender);
 
 	IStaticMeshEditorModule* StaticMeshEditorModule = &FModuleManager::LoadModuleChecked<IStaticMeshEditorModule>("StaticMeshEditor");
+
+	TArray<IStaticMeshEditorModule::FStaticMeshEditorToolbarExtender> ToolbarExtenderDelegates = StaticMeshEditorModule->GetAllStaticMeshEditorToolbarExtenders();
+	for (auto& ToolbarExtenderDelegate : ToolbarExtenderDelegates)
+	{
+		if (ToolbarExtenderDelegate.IsBound())
+		{
+			AddToolbarExtender(ToolbarExtenderDelegate.Execute(GetToolkitCommands(), SharedThis(this)));
+		}
+	}
+
 	EditorToolbarExtender = StaticMeshEditorModule->GetToolBarExtensibilityManager()->GetAllExtenders(GetToolkitCommands(), GetEditingObjects());
 	AddToolbarExtender(EditorToolbarExtender);
 	AddSecondaryToolbarExtender(StaticMeshEditorModule->GetSecondaryToolBarExtensibilityManager()->GetAllExtenders(GetToolkitCommands(), GetEditingObjects()));
@@ -2484,6 +2514,67 @@ TStatId FStaticMeshEditor::GetStatId() const
 	RETURN_QUICK_DECLARE_CYCLE_STAT(FStaticMeshEditor, STATGROUP_TaskGraphTasks);
 }
 
+void FStaticMeshEditor::AddViewportOverlayWidget(TSharedRef<SWidget> InOverlaidWidget)
+{
+	TSharedPtr<SStaticMeshEditorViewport> Viewport = GetStaticMeshViewport();
+
+	if (Viewport.IsValid() && Viewport->GetViewportOverlay().IsValid())
+	{
+		Viewport->GetViewportOverlay()->AddSlot()
+		[
+			InOverlaidWidget
+		];
+	}
+}
+
+void FStaticMeshEditor::RemoveViewportOverlayWidget(TSharedRef<SWidget> InViewportOverlayWidget)
+{
+	TSharedPtr<SStaticMeshEditorViewport> Viewport = GetStaticMeshViewport();
+
+	if (Viewport.IsValid() && Viewport->GetViewportOverlay().IsValid())
+	{
+		Viewport->GetViewportOverlay()->RemoveSlot(InViewportOverlayWidget);
+	}
+}
+
+
+void FStaticMeshEditor::CreateEditorModeManager()
+{
+	//
+	// This function doesn't actually create a new manager -- it assigns StaticMeshEditorViewport->GetViewportClient().GetModeTools() 
+	// to this->EditorModeManager. This is because these two pointers should refer to the same mode manager object, and currently
+	// the ViewPortClient's ModeTools object is created first.
+	// 
+	// This function also:
+	// - sets the manager's PreviewScene to be StaticMeshEditorViewport->GetPreviewScene()
+	// - adds StaticMeshEditorViewport->GetStaticMeshComponent() to the manager's ComponentSet (i.e. selected mesh components)
+	// 
+	
+	TSharedPtr<FAssetEditorModeManager> NewManager = MakeShared<FAssetEditorModeManager>();
+
+	TSharedPtr<SStaticMeshEditorViewport> StaticMeshEditorViewport = GetStaticMeshViewport();
+	if (StaticMeshEditorViewport.IsValid())
+	{
+		TSharedPtr<FEditorModeTools> SharedModeTools = StaticMeshEditorViewport->GetViewportClient().GetModeTools()->AsShared();
+		NewManager = StaticCastSharedPtr<FAssetEditorModeManager>(SharedModeTools);
+		check(NewManager.IsValid());
+
+		TSharedRef<FAdvancedPreviewScene> PreviewScene = StaticMeshEditorViewport->GetPreviewScene();
+		NewManager->SetPreviewScene(&PreviewScene.Get());
+
+		UStaticMeshComponent* const Component = StaticMeshEditorViewport->GetStaticMeshComponent();
+
+		// Copied from FPersonaEditorModeManager::SetPreviewScene(FPreviewScene * NewPreviewScene)
+		USelection* ComponentSet = NewManager->GetSelectedComponents();
+		ComponentSet->BeginBatchSelectOperation();
+		ComponentSet->DeselectAll();
+		ComponentSet->Select(Component, true);
+		ComponentSet->EndBatchSelectOperation();
+	}
+
+	EditorModeManager = NewManager;
+}
+
 bool FStaticMeshEditor::CanRemoveUVChannel()
 {
 	// Can remove UV channel if there's one that is currently being selected and displayed, 
@@ -2904,4 +2995,15 @@ void FStaticMeshEditor::RemoveCurrentUVChannel()
 	}
 }
 
+void FStaticMeshEditor::OnToolkitHostingStarted(const TSharedRef<IToolkit>& Toolkit)
+{
+	ModeUILayer->OnToolkitHostingStarted(Toolkit);
+}
+
+
+void FStaticMeshEditor::OnToolkitHostingFinished(const TSharedRef<IToolkit>& Toolkit)
+{
+	ModeUILayer->OnToolkitHostingFinished(Toolkit);
+}
+
 #undef LOCTEXT_NAMESPACE
diff --git a/Engine/Source/Editor/StaticMeshEditor/Private/StaticMeshEditor.h b/Engine/Source/Editor/StaticMeshEditor/Private/StaticMeshEditor.h
--- a/Engine/Source/Editor/StaticMeshEditor/Private/StaticMeshEditor.h
+++ b/Engine/Source/Editor/StaticMeshEditor/Private/StaticMeshEditor.h
@@ -35,6 +35,7 @@ class UStaticMeshSocket;
 class FViewportTabContent;
 struct FPropertyChangedEvent;
 struct FTabSpawnerEntry;
+class FStaticMeshEditorModeUILayer;
 
 /**
  * StaticMesh Editor class
@@ -69,10 +70,16 @@ private:
 	/** Initializes the editor to use a static mesh. Should be the first thing called. */
 	void InitEditorForStaticMesh(UStaticMesh* ObjectToEdit);
 
+	virtual void PostInitAssetEditor() override;
+
 public:
 	virtual void RegisterTabSpawners(const TSharedRef<class FTabManager>& TabManager) override;
 	virtual void UnregisterTabSpawners(const TSharedRef<class FTabManager>& TabManager) override;
 
+	// IToolkitHost Interface
+	void OnToolkitHostingStarted(const TSharedRef<IToolkit>& Toolkit) override;
+	void OnToolkitHostingFinished(const TSharedRef<IToolkit>& Toolkit) override;
+
 	/**
 	 * Edits the specified static mesh object
 	 *
@@ -217,6 +224,13 @@ public:
 	/** Returns the stat ID for this tickable class */
 	virtual TStatId GetStatId() const final;
 
+	/** Add a widget to the StaticMeshViewport's ViewportOverlay */
+	void AddViewportOverlayWidget(TSharedRef<SWidget> InOverlaidWidget) override;
+
+	/** Remove a widget from the StaticMeshViewport's ViewportOverlay */
+	void RemoveViewportOverlayWidget(TSharedRef<SWidget> InViewportOverlayWidget) override;
+
+	void CreateEditorModeManager() override;
 
 private:
 	TSharedRef<SDockTab> SpawnTab_Viewport(const FSpawnTabArgs& Args);
@@ -576,4 +590,6 @@ private:
 	bool bDrawWireframes;
 	bool bDrawVertexColors;
 	bool bDrawAdditionalData;
+
+	TSharedPtr<FStaticMeshEditorModeUILayer> ModeUILayer;
 };
diff --git a/Engine/Source/Editor/StaticMeshEditor/Private/StaticMeshEditorModeUILayer.cpp b/Engine/Source/Editor/StaticMeshEditor/Private/StaticMeshEditorModeUILayer.cpp
new file mode 100644
--- /dev/null
+++ b/Engine/Source/Editor/StaticMeshEditor/Private/StaticMeshEditorModeUILayer.cpp
@@ -0,0 +1,38 @@
+// Copyright Epic Games, Inc. All Rights Reserved.
+
+#include "StaticMeshEditorModeUILayer.h"
+#include "WorkspaceMenuStructure.h"
+#include "WorkspaceMenuStructureModule.h"
+#include "Toolkits/IToolkit.h"
+
+FStaticMeshEditorModeUILayer::FStaticMeshEditorModeUILayer(const IToolkitHost* InToolkitHost) :
+	FAssetEditorModeUILayer(InToolkitHost)
+{}
+
+void FStaticMeshEditorModeUILayer::OnToolkitHostingStarted(const TSharedRef<IToolkit>& Toolkit)
+{
+	if (!Toolkit->IsAssetEditor())
+	{
+		FAssetEditorModeUILayer::OnToolkitHostingStarted(Toolkit);
+		HostedToolkit = Toolkit;
+		Toolkit->SetModeUILayer(SharedThis(this));
+		Toolkit->RegisterTabSpawners(ToolkitHost->GetTabManager().ToSharedRef());
+		RegisterModeTabSpawners();
+		OnToolkitHostReadyForUI.ExecuteIfBound();
+	}
+}
+
+void FStaticMeshEditorModeUILayer::OnToolkitHostingFinished(const TSharedRef<IToolkit>& Toolkit)
+{
+	if (HostedToolkit.IsValid() && HostedToolkit.Pin() == Toolkit)
+	{
+		FAssetEditorModeUILayer::OnToolkitHostingFinished(Toolkit);
+	}
+}
+
+TSharedPtr<FWorkspaceItem> FStaticMeshEditorModeUILayer::GetModeMenuCategory() const
+{
+	const IWorkspaceMenuStructure& MenuStructure = WorkspaceMenu::GetMenuStructure();
+	return MenuStructure.GetLevelEditorModesCategory();
+}
+
diff --git a/Engine/Source/Editor/StaticMeshEditor/Private/StaticMeshEditorModeUILayer.h b/Engine/Source/Editor/StaticMeshEditor/Private/StaticMeshEditorModeUILayer.h
new file mode 100644
--- /dev/null
+++ b/Engine/Source/Editor/StaticMeshEditor/Private/StaticMeshEditorModeUILayer.h
@@ -0,0 +1,18 @@
+// Copyright Epic Games, Inc. All Rights Reserved.
+
+#pragma once
+
+#include "Toolkits/AssetEditorModeUILayer.h"
+
+class FStaticMeshEditorModeUILayer : public FAssetEditorModeUILayer
+{
+public:
+
+	FStaticMeshEditorModeUILayer(const IToolkitHost* InToolkitHost);
+
+	void OnToolkitHostingStarted(const TSharedRef<IToolkit>& Toolkit) override;
+	void OnToolkitHostingFinished(const TSharedRef<IToolkit>& Toolkit) override;
+	
+	TSharedPtr<FWorkspaceItem> GetModeMenuCategory() const override;
+};
+
diff --git a/Engine/Source/Editor/StaticMeshEditor/Private/StaticMeshEditorModule.cpp b/Engine/Source/Editor/StaticMeshEditor/Private/StaticMeshEditorModule.cpp
--- a/Engine/Source/Editor/StaticMeshEditor/Private/StaticMeshEditorModule.cpp
+++ b/Engine/Source/Editor/StaticMeshEditor/Private/StaticMeshEditorModule.cpp
@@ -56,10 +56,13 @@ public:
 	virtual TSharedPtr<FExtensibilityManager> GetToolBarExtensibilityManager() override { return ToolBarExtensibilityManager; }
 	virtual TSharedPtr<FExtensibilityManager> GetSecondaryToolBarExtensibilityManager() override { return SecondaryToolBarExtensibilityManager; }
 
+	virtual TArray<FStaticMeshEditorToolbarExtender>& GetAllStaticMeshEditorToolbarExtenders() override { return StaticMeshEditorToolbarExtenders; }
+
 private:
 	TSharedPtr<FExtensibilityManager> MenuExtensibilityManager;
 	TSharedPtr<FExtensibilityManager> ToolBarExtensibilityManager;
 	TSharedPtr<FExtensibilityManager> SecondaryToolBarExtensibilityManager;
+	TArray<FStaticMeshEditorToolbarExtender> StaticMeshEditorToolbarExtenders;
 };
 
 IMPLEMENT_MODULE( FStaticMeshEditorModule, StaticMeshEditor );
diff --git a/Engine/Source/Editor/StaticMeshEditor/Public/StaticMeshEditorModule.h b/Engine/Source/Editor/StaticMeshEditor/Public/StaticMeshEditorModule.h
--- a/Engine/Source/Editor/StaticMeshEditor/Public/StaticMeshEditorModule.h
+++ b/Engine/Source/Editor/StaticMeshEditor/Public/StaticMeshEditorModule.h
@@ -29,6 +29,9 @@ public:
 
 	virtual TSharedPtr<FExtensibilityManager> GetSecondaryToolBarExtensibilityManager() = 0;
 
+	DECLARE_DELEGATE_RetVal_TwoParams(TSharedRef<FExtender>, FStaticMeshEditorToolbarExtender, const TSharedRef<FUICommandList>, TSharedRef<IStaticMeshEditor>);
+	virtual TArray<FStaticMeshEditorToolbarExtender>& GetAllStaticMeshEditorToolbarExtenders() = 0;
+
 private:
 	FStaticMeshEditorOpenedEvent StaticMeshEditorOpenedEvent;
 };
diff --git a/Engine/Source/Editor/StaticMeshEditor/StaticMeshEditor.Build.cs b/Engine/Source/Editor/StaticMeshEditor/StaticMeshEditor.Build.cs
--- a/Engine/Source/Editor/StaticMeshEditor/StaticMeshEditor.Build.cs
+++ b/Engine/Source/Editor/StaticMeshEditor/StaticMeshEditor.Build.cs
@@ -46,6 +46,8 @@ public class StaticMeshEditor : ModuleRules
 				"StaticMeshDescription",
 				"ToolMenus",
 				"DetailCustomizations",
+				"StatusBar",
+				"WorkspaceMenuStructure",
 				"PhysicsUtilities",
 			}
 		);
@@ -55,7 +57,6 @@ public class StaticMeshEditor : ModuleRules
 				"SceneOutliner",
 				"ClassViewer",
 				"ContentBrowser",
-				"WorkspaceMenuStructure",
                 "MeshReductionInterface",
             }
 		);
