Add Mesh Inspector tool to Static Mesh Editor's modeling tools

#rnx

diff --git a/Engine/Plugins/Experimental/MeshModelingToolsetExp/Source/MeshModelingToolsExp/Private/MeshInspectorTool.cpp b/Engine/Plugins/Experimental/MeshModelingToolsetExp/Source/MeshModelingToolsExp/Private/MeshInspectorTool.cpp
--- a/Engine/Plugins/Experimental/MeshModelingToolsetExp/Source/MeshModelingToolsExp/Private/MeshInspectorTool.cpp
+++ b/Engine/Plugins/Experimental/MeshModelingToolsetExp/Source/MeshModelingToolsExp/Private/MeshInspectorTool.cpp
@@ -51,7 +51,7 @@ void UMeshInspectorTool::Setup()
 
 	PreviewMesh = NewObject<UPreviewMesh>(this);
 	PreviewMesh->bBuildSpatialDataStructure = false;
-	PreviewMesh->CreateInWorld(UE::ToolTarget::GetTargetActor(Target)->GetWorld(), FTransform::Identity);
+	PreviewMesh->CreateInWorld(TargetWorld, FTransform::Identity);
 	PreviewMesh->SetTransform((FTransform)UE::ToolTarget::GetLocalToWorldTransform(Target));
 
 	FComponentMaterialSet MaterialSet = UE::ToolTarget::GetMaterialSet(Target);
diff --git a/Engine/Plugins/Experimental/StaticMeshEditorModeling/Source/StaticMeshEditorModeling/Private/StaticMeshEditorModelingMode.cpp b/Engine/Plugins/Experimental/StaticMeshEditorModeling/Source/StaticMeshEditorModeling/Private/StaticMeshEditorModelingMode.cpp
--- a/Engine/Plugins/Experimental/StaticMeshEditorModeling/Source/StaticMeshEditorModeling/Private/StaticMeshEditorModelingMode.cpp
+++ b/Engine/Plugins/Experimental/StaticMeshEditorModeling/Source/StaticMeshEditorModeling/Private/StaticMeshEditorModelingMode.cpp
@@ -8,6 +8,7 @@
 #include "ToolTargets/StaticMeshComponentToolTarget.h"
 #include "Tools/GenerateStaticMeshLODAssetTool.h"
 #include "Tools/LODManagerTool.h"
+#include "MeshInspectorTool.h"
 
 #define LOCTEXT_NAMESPACE "StaticMeshEditorModelingMode"
 
@@ -28,6 +29,7 @@ void UStaticMeshEditorModelingMode::Enter()
 
 	RegisterTool(ToolManagerCommands.BeginGenerateStaticMeshLODAssetTool, TEXT("BeginGenerateStaticMeshLODAssetTool"), NewObject<UGenerateStaticMeshLODAssetToolBuilder>());
 	RegisterTool(ToolManagerCommands.BeginLODManagerTool, TEXT("BeginLODManagerTool"), NewObject<ULODManagerToolBuilder>());
+	RegisterTool(ToolManagerCommands.BeginMeshInspectorTool, TEXT("BeginMeshInspectorTool"), NewObject<UMeshInspectorToolBuilder>());
 }
 
 bool UStaticMeshEditorModelingMode::UsesToolkits() const
diff --git a/Engine/Plugins/Experimental/StaticMeshEditorModeling/Source/StaticMeshEditorModeling/Private/StaticMeshEditorModelingToolkit.cpp b/Engine/Plugins/Experimental/StaticMeshEditorModeling/Source/StaticMeshEditorModeling/Private/StaticMeshEditorModelingToolkit.cpp
--- a/Engine/Plugins/Experimental/StaticMeshEditorModeling/Source/StaticMeshEditorModeling/Private/StaticMeshEditorModelingToolkit.cpp
+++ b/Engine/Plugins/Experimental/StaticMeshEditorModeling/Source/StaticMeshEditorModeling/Private/StaticMeshEditorModelingToolkit.cpp
@@ -95,14 +95,14 @@ void FStaticMeshEditorModelingToolkit::Init(
 		[
 			SNew(SHorizontalBox)
 
-			//+SHorizontalBox::Slot()
-			//.AutoWidth()
-			//.VAlign(VAlign_Center)
-			//.Padding(FMargin(0.f, 0.f, 8.f, 0.f))
-			//[
-			//	SNew(SImage)
-			//	.Image_Lambda([this] () { return ActiveToolIcon; })
-			//]
+			+SHorizontalBox::Slot()
+			.AutoWidth()
+			.VAlign(VAlign_Center)
+			.Padding(FMargin(0.f, 0.f, 8.f, 0.f))
+			[
+				SNew(SImage)
+				.Image_Lambda([this] () { return ActiveToolIcon; })
+			]
 
 			+SHorizontalBox::Slot()
 			.AutoWidth()
@@ -143,20 +143,20 @@ void FStaticMeshEditorModelingToolkit::Init(
 				.Visibility_Lambda([this]() { return GetScriptableEditorMode()->GetInteractiveToolsContext()->ActiveToolHasAccept() ? EVisibility::Visible : EVisibility::Collapsed; })
 			]
 
-			//+SHorizontalBox::Slot()
-			//.AutoWidth()
-			//.Padding(FMargin(2.0, 0.f, 0.f, 0.f))
-			//[
-			//	SNew(SButton)
-			//	.ButtonStyle(FAppStyle::Get(), "PrimaryButton")
-			//	.TextStyle( FAppStyle::Get(), "DialogButtonText" )
-			//	.Text(LOCTEXT("OverlayComplete", "Complete"))
-			//	.ToolTipText(LOCTEXT("OverlayCompleteTooltip", "Exit the active Tool [Enter]"))
-			//	.HAlign(HAlign_Center)
-			//	.OnClicked_Lambda([this]() { GetScriptableEditorMode()->GetInteractiveToolsContext()->EndTool(EToolShutdownType::Completed); return FReply::Handled(); })
-			//	.IsEnabled_Lambda([this]() { return GetScriptableEditorMode()->GetInteractiveToolsContext()->CanCompleteActiveTool(); })
-			//	.Visibility_Lambda([this]() { return GetScriptableEditorMode()->GetInteractiveToolsContext()->CanCompleteActiveTool() ? EVisibility::Visible : EVisibility::Collapsed; })
-			//]
+			+SHorizontalBox::Slot()
+			.AutoWidth()
+			.Padding(FMargin(2.0, 0.f, 0.f, 0.f))
+			[
+				SNew(SButton)
+				.ButtonStyle(FAppStyle::Get(), "PrimaryButton")
+				.TextStyle( FAppStyle::Get(), "DialogButtonText" )
+				.Text(LOCTEXT("OverlayComplete", "Complete"))
+				.ToolTipText(LOCTEXT("OverlayCompleteTooltip", "Exit the active Tool [Enter]"))
+				.HAlign(HAlign_Center)
+				.OnClicked_Lambda([this]() { GetScriptableEditorMode()->GetInteractiveToolsContext()->EndTool(EToolShutdownType::Completed); return FReply::Handled(); })
+				.IsEnabled_Lambda([this]() { return GetScriptableEditorMode()->GetInteractiveToolsContext()->CanCompleteActiveTool(); })
+				.Visibility_Lambda([this]() { return GetScriptableEditorMode()->GetInteractiveToolsContext()->CanCompleteActiveTool() ? EVisibility::Visible : EVisibility::Collapsed; })
+			]
 		]	
 	];
 
@@ -199,6 +199,7 @@ void FStaticMeshEditorModelingToolkit::BuildToolPalette(FName PaletteName, class
 	{
 		ToolbarBuilder.AddToolBarButton(Commands.BeginGenerateStaticMeshLODAssetTool);
 		ToolbarBuilder.AddToolBarButton(Commands.BeginLODManagerTool);
+		ToolbarBuilder.AddToolBarButton(Commands.BeginMeshInspectorTool);
 	}
 }
 
diff --git a/Engine/Plugins/Experimental/StaticMeshEditorModeling/Source/StaticMeshEditorModeling/StaticMeshEditorModeling.Build.cs b/Engine/Plugins/Experimental/StaticMeshEditorModeling/Source/StaticMeshEditorModeling/StaticMeshEditorModeling.Build.cs
--- a/Engine/Plugins/Experimental/StaticMeshEditorModeling/Source/StaticMeshEditorModeling/StaticMeshEditorModeling.Build.cs
+++ b/Engine/Plugins/Experimental/StaticMeshEditorModeling/Source/StaticMeshEditorModeling/StaticMeshEditorModeling.Build.cs
@@ -45,6 +45,7 @@ public class StaticMeshEditorModeling : ModuleRules
 				"Slate",
 				"SlateCore",
 				"StaticMeshEditor",
+				"MeshModelingToolsExp",
 				"ToolMenus",
 				"UnrealEd",
 			}
