Warn the user that new assets might be created when AutoLOD is launched from the Static Mesh Editor

#rnx

diff --git a/Engine/Plugins/Experimental/StaticMeshEditorModeling/Source/StaticMeshEditorModeling/Private/StaticMeshEditorModelingToolkit.cpp b/Engine/Plugins/Experimental/StaticMeshEditorModeling/Source/StaticMeshEditorModeling/Private/StaticMeshEditorModelingToolkit.cpp
--- a/Engine/Plugins/Experimental/StaticMeshEditorModeling/Source/StaticMeshEditorModeling/Private/StaticMeshEditorModelingToolkit.cpp
+++ b/Engine/Plugins/Experimental/StaticMeshEditorModeling/Source/StaticMeshEditorModeling/Private/StaticMeshEditorModelingToolkit.cpp
@@ -13,7 +13,7 @@
 #include "EditorStyleSet.h"
 #include "IDetailsView.h"
 #include "Framework/MultiBox/MultiBoxBuilder.h"
-
+#include "Tools/GenerateStaticMeshLODAssetTool.h"
 
 #define LOCTEXT_NAMESPACE "StaticMeshEditorModelingToolkit"
 
@@ -260,6 +260,12 @@ void FStaticMeshEditorModelingToolkit::OnToolStarted(UInteractiveToolManager* Ma
 	ActiveToolIcon = FModelingToolsEditorModeStyle::Get()->GetOptionalBrush(ActiveToolIconName);
 
 	GetToolkitHost()->AddViewportOverlayWidget(ViewportOverlayWidget.ToSharedRef());
+
+	const UInteractiveTool* CurrentTool = GetScriptableEditorMode()->GetToolManager()->GetActiveTool(EToolSide::Left);
+	if (CurrentTool->GetClass()->IsChildOf<UGenerateStaticMeshLODAssetTool>())
+	{
+		PostWarning(LOCTEXT("AutoLODStaticMeshEditorWarning", "Note: This tool may create new assets including new Static Meshes, Textures, and Materials"));
+	}
 }
 
 
