Fix errors in clang build, related to new StaticMeshEditorModeling plugin

#rnx

diff --git a/Engine/Plugins/Experimental/StaticMeshEditorModeling/Source/StaticMeshEditorModeling/StaticMeshEditorModeling.Build.cs b/Engine/Plugins/Experimental/StaticMeshEditorModeling/Source/StaticMeshEditorModeling/StaticMeshEditorModeling.Build.cs
--- a/Engine/Plugins/Experimental/StaticMeshEditorModeling/Source/StaticMeshEditorModeling/StaticMeshEditorModeling.Build.cs
+++ b/Engine/Plugins/Experimental/StaticMeshEditorModeling/Source/StaticMeshEditorModeling/StaticMeshEditorModeling.Build.cs
@@ -34,6 +34,7 @@ public class StaticMeshEditorModeling : ModuleRules
 			new string[]
 			{
 				"Core",
+				"InputCore",
 				"CoreUObject",
 				"EditorInteractiveToolsFramework",
 				"EditorStyle",
