Better icon and tab name for mesh modeling mode in Static Mesh Editor

#rnx

diff --git a/Engine/Plugins/Experimental/StaticMeshEditorModeling/Source/StaticMeshEditorModeling/Private/StaticMeshEditorModelingMode.cpp b/Engine/Plugins/Experimental/StaticMeshEditorModeling/Source/StaticMeshEditorModeling/Private/StaticMeshEditorModelingMode.cpp
--- a/Engine/Plugins/Experimental/StaticMeshEditorModeling/Source/StaticMeshEditorModeling/Private/StaticMeshEditorModelingMode.cpp
+++ b/Engine/Plugins/Experimental/StaticMeshEditorModeling/Source/StaticMeshEditorModeling/Private/StaticMeshEditorModelingMode.cpp
@@ -16,7 +16,7 @@ const FEditorModeID UStaticMeshEditorModelingMode::Id("StaticMeshEditorModelingM
 
 UStaticMeshEditorModelingMode::UStaticMeshEditorModelingMode()
 {
-	Info = FEditorModeInfo(Id, LOCTEXT("StaticMeshEditorModelingMode", "StaticMeshEditorModelingMode"), FSlateIcon(), true);
+	Info = FEditorModeInfo(Id, LOCTEXT("StaticMeshEditorModelingMode", "Modeling"), FSlateIcon("ModelingToolsStyle", "LevelEditor.ModelingToolsMode"), true);
 }
 
 void UStaticMeshEditorModelingMode::Enter()
diff --git a/Engine/Plugins/Experimental/StaticMeshEditorModeling/Source/StaticMeshEditorModeling/Private/StaticMeshEditorModelingModule.cpp b/Engine/Plugins/Experimental/StaticMeshEditorModeling/Source/StaticMeshEditorModeling/Private/StaticMeshEditorModelingModule.cpp
--- a/Engine/Plugins/Experimental/StaticMeshEditorModeling/Source/StaticMeshEditorModeling/Private/StaticMeshEditorModelingModule.cpp
+++ b/Engine/Plugins/Experimental/StaticMeshEditorModeling/Source/StaticMeshEditorModeling/Private/StaticMeshEditorModelingModule.cpp
@@ -80,9 +80,9 @@ void FStaticMeshEditorModelingModule::RegisterMenusAndToolbars()
 		{
 			InSection.AddEntry(FToolMenuEntry::InitToolBarButton(
 				FStaticMeshEditorModelingCommands::Get().ToggleStaticMeshEditorModelingMode,
-				LOCTEXT("StaticMeshEditorModelingMode", "Mesh Modeling Tools"),
+				LOCTEXT("StaticMeshEditorModelingMode", "Modeling Tools"),
 				FText(),
-				FSlateIcon()));
+				FSlateIcon("ModelingToolsStyle", "LevelEditor.ModelingToolsMode")));
 		}));
 	}
 }
