ModelingMode: make the StaticMeshEditorModeling plugin visible in the plugins list, so that it can be turned on by end-users
#rnx

diff --git a/Engine/Plugins/Experimental/StaticMeshEditorModeling/StaticMeshEditorModeling.uplugin b/Engine/Plugins/Experimental/StaticMeshEditorModeling/StaticMeshEditorModeling.uplugin
--- a/Engine/Plugins/Experimental/StaticMeshEditorModeling/StaticMeshEditorModeling.uplugin
+++ b/Engine/Plugins/Experimental/StaticMeshEditorModeling/StaticMeshEditorModeling.uplugin
@@ -2,9 +2,9 @@
 	"FileVersion": 3,
 	"Version": 1,
 	"VersionName": "0.1",
-	"FriendlyName": "Mesh LOD Plugin",
-	"Description": "A set of modules implementing 3D mesh LOD creation",
-	"Category": "Other",
+	"FriendlyName": "Static Mesh Editor Modeling Mode",
+	"Description": "Enable a Modeling Tools Tab in the Static Mesh Editor",
+	"Category": "Editor",
 	"CreatedBy": "Epic Games, Inc.",
 	"CreatedByURL": "http://epicgames.com",
 	"DocsURL": "",
@@ -13,7 +13,7 @@
 	"CanContainContent": true,
 	"IsBetaVersion": true,
 	"Installed": false,
-	"Hidden": true,
+	"Hidden": false,
 	"Modules": [
 		{
 			"Name": "StaticMeshEditorModeling",
