Static Mesh Editor: prompt user to Accept/Cancel current modeling tool (if any) when shutting down SME

#rnx

diff --git a/Engine/Plugins/Experimental/StaticMeshEditorModeling/Source/StaticMeshEditorModeling/Private/StaticMeshEditorModelingMode.cpp b/Engine/Plugins/Experimental/StaticMeshEditorModeling/Source/StaticMeshEditorModeling/Private/StaticMeshEditorModelingMode.cpp
--- a/Engine/Plugins/Experimental/StaticMeshEditorModeling/Source/StaticMeshEditorModeling/Private/StaticMeshEditorModelingMode.cpp
+++ b/Engine/Plugins/Experimental/StaticMeshEditorModeling/Source/StaticMeshEditorModeling/Private/StaticMeshEditorModelingMode.cpp
@@ -9,6 +9,7 @@
 #include "Tools/GenerateStaticMeshLODAssetTool.h"
 #include "Tools/LODManagerTool.h"
 #include "MeshInspectorTool.h"
+#include "Misc/MessageDialog.h"
 
 #define LOCTEXT_NAMESPACE "StaticMeshEditorModelingMode"
 
@@ -34,6 +35,47 @@ void UStaticMeshEditorModelingMode::Enter()
 	RegisterTool(ToolManagerCommands.BeginMeshInspectorTool, TEXT("BeginMeshInspectorTool"), NewObject<UMeshInspectorToolBuilder>());
 }
 
+bool UStaticMeshEditorModelingMode::OnRequestClose()
+{
+	bool bAllowClose = true;
+
+	if (GetInteractiveToolsContext()->HasActiveTool())
+	{
+		UInteractiveToolsContext* ITC = GetInteractiveToolsContext();
+
+		if (GetInteractiveToolsContext()->CanAcceptActiveTool())
+		{
+			EAppReturnType::Type YesNoCancelReply = FMessageDialog::Open(
+				EAppMsgType::YesNoCancel,
+				LOCTEXT("ShouldApplyModelingTool", "Would you like to Accept the current modeling tool changes?")
+			);
+
+			switch (YesNoCancelReply)
+			{
+			case EAppReturnType::Yes:
+				ITC->DeactivateAllActiveTools(EToolShutdownType::Accept);
+				bAllowClose = true;
+				break;
+			case EAppReturnType::No:
+				ITC->DeactivateAllActiveTools(EToolShutdownType::Cancel);
+				bAllowClose = true;
+				break;
+			case EAppReturnType::Cancel:
+				bAllowClose = false;
+				break;
+			}
+		}
+		else
+		{
+			// No Accept option, so cancel and close
+			ITC->DeactivateAllActiveTools(EToolShutdownType::Cancel);
+			bAllowClose = true;
+		}
+	}
+
+	return bAllowClose;
+}
+
 bool UStaticMeshEditorModelingMode::UsesToolkits() const
 { 
 	return true; 
diff --git a/Engine/Plugins/Experimental/StaticMeshEditorModeling/Source/StaticMeshEditorModeling/Public/StaticMeshEditorModelingMode.h b/Engine/Plugins/Experimental/StaticMeshEditorModeling/Source/StaticMeshEditorModeling/Public/StaticMeshEditorModelingMode.h
--- a/Engine/Plugins/Experimental/StaticMeshEditorModeling/Source/StaticMeshEditorModeling/Public/StaticMeshEditorModelingMode.h
+++ b/Engine/Plugins/Experimental/StaticMeshEditorModeling/Source/StaticMeshEditorModeling/Public/StaticMeshEditorModelingMode.h
@@ -21,10 +21,11 @@ protected:
 
 	UStaticMeshEditorModelingMode();
 
-	void Enter() override;
+	virtual void Enter() override;
+	virtual bool OnRequestClose() override;
 
-	bool UsesToolkits() const override; 
-	void CreateToolkit() override;
+	virtual bool UsesToolkits() const override;
+	virtual void CreateToolkit() override;
 	
 
 };
diff --git a/Engine/Source/Editor/StaticMeshEditor/Private/StaticMeshEditor.cpp b/Engine/Source/Editor/StaticMeshEditor/Private/StaticMeshEditor.cpp
--- a/Engine/Source/Editor/StaticMeshEditor/Private/StaticMeshEditor.cpp
+++ b/Engine/Source/Editor/StaticMeshEditor/Private/StaticMeshEditor.cpp
@@ -2398,6 +2398,8 @@ bool FStaticMeshEditor::OnRequestClose()
 		}
 	}
 
+	bAllowClose &= GetEditorModeManager().OnRequestClose();
+
 	// Give any active modes a chance to shutdown while the toolkit host is still alive
 	if (bAllowClose)
 	{
diff --git a/Engine/Source/Editor/UnrealEd/Private/EditorModeManager.cpp b/Engine/Source/Editor/UnrealEd/Private/EditorModeManager.cpp
--- a/Engine/Source/Editor/UnrealEd/Private/EditorModeManager.cpp
+++ b/Engine/Source/Editor/UnrealEd/Private/EditorModeManager.cpp
@@ -1838,6 +1838,25 @@ bool FEditorModeTools::CanAutoSave() const
 	return bCanAutoSave;
 }
 
+
+bool FEditorModeTools::OnRequestClose()
+{
+	bool bCanClose = true;
+
+	ForEachEdMode([&bCanClose](UEdMode* Mode)
+	{
+		if (!Mode->OnRequestClose())
+		{
+			bCanClose = false;
+		}
+
+		return true;
+	});
+
+	return bCanClose;
+
+}
+
 UModeManagerInteractiveToolsContext* FEditorModeTools::GetInteractiveToolsContext() const
 {
 	return InteractiveToolsContext;
diff --git a/Engine/Source/Editor/UnrealEd/Public/EditorModeManager.h b/Engine/Source/Editor/UnrealEd/Public/EditorModeManager.h
--- a/Engine/Source/Editor/UnrealEd/Public/EditorModeManager.h
+++ b/Engine/Source/Editor/UnrealEd/Public/EditorModeManager.h
@@ -405,6 +405,9 @@ public:
 	/** Is the viewport UI hidden? */
 	bool IsViewportUIHidden() const { return bHideViewportUI; }
 
+	/** Called by Editors when they are about to close */
+	bool OnRequestClose();
+
 	bool PivotShown;
 	bool Snapping;
 	bool SnappedActor;
diff --git a/Engine/Source/Editor/UnrealEd/Public/Tools/UEdMode.h b/Engine/Source/Editor/UnrealEd/Public/Tools/UEdMode.h
--- a/Engine/Source/Editor/UnrealEd/Public/Tools/UEdMode.h
+++ b/Engine/Source/Editor/UnrealEd/Public/Tools/UEdMode.h
@@ -183,6 +183,11 @@ public:
 	*/
 	virtual FEdMode* AsLegacyMode() { return nullptr; }
 	
+	virtual bool OnRequestClose() 
+	{
+		return true;
+	}
+
 protected:
 
 	/** Information pertaining to this mode. Should be assigned in the constructor. */
