Resolve keyboard shortcut issues for Landscape and Modeling in the Editor Preferences

diff --git a/Engine/Plugins/Editor/ModelingToolsEditorMode/Source/ModelingToolsEditorMode/Private/ModelingToolsActions.cpp b/Engine/Plugins/Editor/ModelingToolsEditorMode/Source/ModelingToolsEditorMode/Private/ModelingToolsActions.cpp
--- a/Engine/Plugins/Editor/ModelingToolsEditorMode/Source/ModelingToolsEditorMode/Private/ModelingToolsActions.cpp
+++ b/Engine/Plugins/Editor/ModelingToolsEditorMode/Source/ModelingToolsEditorMode/Private/ModelingToolsActions.cpp
@@ -27,7 +27,7 @@
 FModelingModeActionCommands::FModelingModeActionCommands() :
 	TCommands<FModelingModeActionCommands>(
 		"ModelingModeCommands", // Context name for fast lookup
-		NSLOCTEXT("Contexts", "ModelingModeCommands", "Modeling Mode Shortcuts"), // Localized context name for displaying
+		NSLOCTEXT("Contexts", "ModelingModeCommands", "Modeling Mode"), // Localized context name for displaying
 		NAME_None, // Parent
 		FAppStyle::GetAppStyleSetName() // Icon Style Set
 		)
diff --git a/Engine/Plugins/Experimental/StaticMeshEditorModeling/Source/StaticMeshEditorModeling/Private/StaticMeshEditorModelingCommands.cpp b/Engine/Plugins/Experimental/StaticMeshEditorModeling/Source/StaticMeshEditorModeling/Private/StaticMeshEditorModelingCommands.cpp
--- a/Engine/Plugins/Experimental/StaticMeshEditorModeling/Source/StaticMeshEditorModeling/Private/StaticMeshEditorModelingCommands.cpp
+++ b/Engine/Plugins/Experimental/StaticMeshEditorModeling/Source/StaticMeshEditorModeling/Private/StaticMeshEditorModelingCommands.cpp
@@ -7,7 +7,7 @@
 #define LOCTEXT_NAMESPACE "StaticMeshEditorModelingCommands"
 
 FStaticMeshEditorModelingCommands::FStaticMeshEditorModelingCommands()
-	: TCommands<FStaticMeshEditorModelingCommands>(TEXT("MeshModelingTools"), NSLOCTEXT("Contexts", "MeshModelingTools", "Mesh Modeling"), NAME_None, FAppStyle::GetAppStyleSetName())
+	: TCommands<FStaticMeshEditorModelingCommands>(TEXT("MeshModelingTools"), NSLOCTEXT("Contexts", "MeshModelingTools", "Static Mesh Editor - Modeling Tools"), NAME_None, FAppStyle::GetAppStyleSetName())
 {}
 
 void FStaticMeshEditorModelingCommands::RegisterCommands()
