Move StaticMeshEditorModeling into Editor plugins folder

diff --git a/Engine/Plugins/Editor/StaticMeshEditorModeling/Source/StaticMeshEditorModeling/Private/StaticMeshEditorModelingCommands.cpp b/Engine/Plugins/Editor/StaticMeshEditorModeling/Source/StaticMeshEditorModeling/Private/StaticMeshEditorModelingCommands.cpp
new file mode 100644
--- /dev/null
+++ b/Engine/Plugins/Editor/StaticMeshEditorModeling/Source/StaticMeshEditorModeling/Private/StaticMeshEditorModelingCommands.cpp
@@ -0,0 +1,23 @@
+// Copyright Epic Games, Inc. All Rights Reserved.
+
+#include "StaticMeshEditorModelingCommands.h"
+#include "Styling/AppStyle.h"
+#include "Framework/Commands/UICommandInfo.h"
+
+#define LOCTEXT_NAMESPACE "StaticMeshEditorModelingCommands"
+
+FStaticMeshEditorModelingCommands::FStaticMeshEditorModelingCommands()
+	: TCommands<FStaticMeshEditorModelingCommands>(TEXT("MeshModelingTools"), NSLOCTEXT("Contexts", "MeshModelingTools", "Static Mesh Editor - Modeling Tools"), NAME_None, FAppStyle::GetAppStyleSetName())
+{}
+
+void FStaticMeshEditorModelingCommands::RegisterCommands()
+{
+	UI_COMMAND(ToggleStaticMeshEditorModelingMode, "Enable Mesh Modeling Tools", "Toggles Mesh Modeling Tools on or off.", EUserInterfaceActionType::ToggleButton, FInputChord());
+}
+
+const FStaticMeshEditorModelingCommands& FStaticMeshEditorModelingCommands::Get()
+{
+	return TCommands<FStaticMeshEditorModelingCommands>::Get();
+}
+
+#undef LOCTEXT_NAMESPACE
diff --git a/Engine/Plugins/Editor/StaticMeshEditorModeling/Source/StaticMeshEditorModeling/Private/StaticMeshEditorModelingCommands.h b/Engine/Plugins/Editor/StaticMeshEditorModeling/Source/StaticMeshEditorModeling/Private/StaticMeshEditorModelingCommands.h
new file mode 100644
--- /dev/null
+++ b/Engine/Plugins/Editor/StaticMeshEditorModeling/Source/StaticMeshEditorModeling/Private/StaticMeshEditorModelingCommands.h
@@ -0,0 +1,16 @@
+// Copyright Epic Games, Inc. All Rights Reserved.
+
+#pragma once
+
+#include "Framework/Commands/Commands.h"
+
+class FStaticMeshEditorModelingCommands : public TCommands<FStaticMeshEditorModelingCommands>
+{
+public:
+	FStaticMeshEditorModelingCommands();
+	void RegisterCommands() override;
+	static const FStaticMeshEditorModelingCommands& Get();
+
+	TSharedPtr<FUICommandInfo> ToggleStaticMeshEditorModelingMode;
+};
+
diff --git a/Engine/Plugins/Editor/StaticMeshEditorModeling/Source/StaticMeshEditorModeling/Private/StaticMeshEditorModelingMode.cpp b/Engine/Plugins/Editor/StaticMeshEditorModeling/Source/StaticMeshEditorModeling/Private/StaticMeshEditorModelingMode.cpp
new file mode 100644
--- /dev/null
+++ b/Engine/Plugins/Editor/StaticMeshEditorModeling/Source/StaticMeshEditorModeling/Private/StaticMeshEditorModelingMode.cpp
@@ -0,0 +1,93 @@
+// Copyright Epic Games, Inc. All Rights Reserved.
+
+#include "StaticMeshEditorModelingMode.h"
+#include "StaticMeshEditorModelingToolkit.h"
+#include "EdModeInteractiveToolsContext.h"
+#include "ToolTargetManager.h"
+#include "ModelingToolsManagerActions.h"
+#include "ModelingToolsEditorModeSettings.h"
+#include "ToolTargets/StaticMeshComponentToolTarget.h"
+#include "Tools/GenerateStaticMeshLODAssetTool.h"
+#include "Tools/LODManagerTool.h"
+#include "MeshInspectorTool.h"
+#include "Misc/MessageDialog.h"
+
+#define LOCTEXT_NAMESPACE "StaticMeshEditorModelingMode"
+
+const FEditorModeID UStaticMeshEditorModelingMode::Id("StaticMeshEditorModelingMode");
+
+UStaticMeshEditorModelingMode::UStaticMeshEditorModelingMode()
+{
+	Info = FEditorModeInfo(Id, LOCTEXT("StaticMeshEditorModelingMode", "Modeling"), FSlateIcon("ModelingToolsStyle", "LevelEditor.ModelingToolsMode"), false);
+}
+
+void UStaticMeshEditorModelingMode::Enter()
+{
+	UEdMode::Enter();
+
+	GetInteractiveToolsContext()->TargetManager->AddTargetFactory(NewObject<UStaticMeshComponentToolTargetFactory>(GetToolManager()));
+
+	const FModelingToolsManagerCommands& ToolManagerCommands = FModelingToolsManagerCommands::Get();
+
+	const UModelingToolsEditorModeSettings* ModelingModeSettings = GetDefault<UModelingToolsEditorModeSettings>();
+
+	UGenerateStaticMeshLODAssetToolBuilder* GenerateStaticMeshLODAssetToolBuilder = NewObject<UGenerateStaticMeshLODAssetToolBuilder>();
+	GenerateStaticMeshLODAssetToolBuilder->bUseAssetEditorMode = true;
+	GenerateStaticMeshLODAssetToolBuilder->bInRestrictiveMode = ModelingModeSettings->InRestrictiveMode();
+	RegisterTool(ToolManagerCommands.BeginGenerateStaticMeshLODAssetTool, TEXT("BeginGenerateStaticMeshLODAssetTool"), GenerateStaticMeshLODAssetToolBuilder);
+	RegisterTool(ToolManagerCommands.BeginLODManagerTool, TEXT("BeginLODManagerTool"), NewObject<ULODManagerToolBuilder>());
+	RegisterTool(ToolManagerCommands.BeginMeshInspectorTool, TEXT("BeginMeshInspectorTool"), NewObject<UMeshInspectorToolBuilder>());
+}
+
+bool UStaticMeshEditorModelingMode::OnRequestClose()
+{
+	bool bAllowClose = true;
+
+	if (GetInteractiveToolsContext()->HasActiveTool())
+	{
+		UInteractiveToolsContext* ITC = GetInteractiveToolsContext();
+
+		if (GetInteractiveToolsContext()->CanAcceptActiveTool())
+		{
+			EAppReturnType::Type YesNoCancelReply = FMessageDialog::Open(
+				EAppMsgType::YesNoCancel,
+				LOCTEXT("ShouldApplyModelingTool", "Would you like to Accept the current modeling tool changes?")
+			);
+
+			switch (YesNoCancelReply)
+			{
+			case EAppReturnType::Yes:
+				ITC->DeactivateAllActiveTools(EToolShutdownType::Accept);
+				bAllowClose = true;
+				break;
+			case EAppReturnType::No:
+				ITC->DeactivateAllActiveTools(EToolShutdownType::Cancel);
+				bAllowClose = true;
+				break;
+			case EAppReturnType::Cancel:
+				bAllowClose = false;
+				break;
+			}
+		}
+		else
+		{
+			// No Accept option, so cancel and close
+			ITC->DeactivateAllActiveTools(EToolShutdownType::Cancel);
+			bAllowClose = true;
+		}
+	}
+
+	return bAllowClose;
+}
+
+bool UStaticMeshEditorModelingMode::UsesToolkits() const
+{ 
+	return true; 
+}
+
+void UStaticMeshEditorModelingMode::CreateToolkit()
+{
+	Toolkit = MakeShareable(new FStaticMeshEditorModelingToolkit);
+}
+
+#undef LOCTEXT_NAMESPACE
diff --git a/Engine/Plugins/Editor/StaticMeshEditorModeling/Source/StaticMeshEditorModeling/Private/StaticMeshEditorModelingModule.cpp b/Engine/Plugins/Editor/StaticMeshEditorModeling/Source/StaticMeshEditorModeling/Private/StaticMeshEditorModelingModule.cpp
new file mode 100644
--- /dev/null
+++ b/Engine/Plugins/Editor/StaticMeshEditorModeling/Source/StaticMeshEditorModeling/Private/StaticMeshEditorModelingModule.cpp
@@ -0,0 +1,94 @@
+// Copyright Epic Games, Inc. All Rights Reserved.
+
+#include "StaticMeshEditorModelingModule.h"
+#include "StaticMeshEditorModelingCommands.h"
+#include "StaticMeshEditorModelingMode.h"
+#include "ToolMenus.h"
+#include "StaticMeshEditorModule.h"
+#include "IStaticMeshEditor.h"
+#include "EditorModeManager.h"
+
+#define LOCTEXT_NAMESPACE "FStaticMeshEditorModelingModule"
+
+void FStaticMeshEditorModelingModule::StartupModule()
+{
+	// This code will execute after your module is loaded into memory; the exact timing is specified in the .uplugin file per-module
+	FStaticMeshEditorModelingCommands::Register();
+
+	UToolMenus::RegisterStartupCallback(
+		FSimpleMulticastDelegate::FDelegate::CreateRaw(this, &FStaticMeshEditorModelingModule::RegisterMenusAndToolbars));
+
+	IStaticMeshEditorModule& StaticMeshEditorModule = FModuleManager::LoadModuleChecked<IStaticMeshEditorModule>("StaticMeshEditor");
+	TArray<IStaticMeshEditorModule::FStaticMeshEditorToolbarExtender>& ToolbarExtenders = StaticMeshEditorModule.GetAllStaticMeshEditorToolbarExtenders();
+
+	// Add a button to activate modeling mode to the static mesh editor toolbar
+	auto ExtendStaticMeshEditorToolbar = [this](const TSharedRef<FUICommandList> InCommandList, TSharedRef<IStaticMeshEditor> Editor) -> TSharedRef<FExtender>
+	{
+		TSharedPtr<FExtender> ToolbarExtender = MakeShareable(new FExtender);
+		TWeakPtr<IStaticMeshEditor> Ptr(Editor);
+
+		InCommandList->MapAction(FStaticMeshEditorModelingCommands::Get().ToggleStaticMeshEditorModelingMode,
+			FExecuteAction::CreateRaw(this, &FStaticMeshEditorModelingModule::OnToggleStaticMeshEditorModelingMode, Ptr),
+			FCanExecuteAction(),
+			FIsActionChecked::CreateRaw(this, &FStaticMeshEditorModelingModule::IsStaticMeshEditorModelingModeActive, Ptr));
+
+		return ToolbarExtender.ToSharedRef();
+	};
+
+	ToolbarExtenders.Add(IStaticMeshEditorModule::FStaticMeshEditorToolbarExtender::CreateLambda(ExtendStaticMeshEditorToolbar));
+}
+
+
+void FStaticMeshEditorModelingModule::ShutdownModule()
+{
+	UToolMenus::UnregisterOwner(this);
+	FStaticMeshEditorModelingCommands::Unregister();
+}
+
+bool FStaticMeshEditorModelingModule::IsStaticMeshEditorModelingModeActive(TWeakPtr<IStaticMeshEditor> InEditor) const
+{
+	TSharedPtr<IStaticMeshEditor> PinnedEditor = InEditor.Pin();
+	return PinnedEditor.IsValid() && PinnedEditor->GetEditorModeManager().IsModeActive(UStaticMeshEditorModelingMode::Id);
+}
+
+
+void FStaticMeshEditorModelingModule::OnToggleStaticMeshEditorModelingMode(TWeakPtr<IStaticMeshEditor> InEditor)
+{
+	TSharedPtr<IStaticMeshEditor> PinnedEditor = InEditor.Pin();
+	if (PinnedEditor.IsValid())
+	{
+		if (!IsStaticMeshEditorModelingModeActive(InEditor))
+		{
+			PinnedEditor->GetEditorModeManager().ActivateMode(UStaticMeshEditorModelingMode::Id, true);
+		}
+		else
+		{
+			PinnedEditor->GetEditorModeManager().ActivateDefaultMode();
+		}
+	}
+}
+
+void FStaticMeshEditorModelingModule::RegisterMenusAndToolbars()
+{
+	FToolMenuOwnerScoped OwnerScoped(this);
+
+	UToolMenu* Toolbar = UToolMenus::Get()->ExtendMenu("AssetEditor.StaticMeshEditor.ToolBar");
+
+	{
+		FToolMenuSection& Section = Toolbar->FindOrAddSection("StaticMesh");
+		Section.AddDynamicEntry("ToggleStaticMeshEditorModelingMode", FNewToolMenuSectionDelegate::CreateLambda([this](FToolMenuSection& InSection)
+		{
+			InSection.AddEntry(FToolMenuEntry::InitToolBarButton(
+				FStaticMeshEditorModelingCommands::Get().ToggleStaticMeshEditorModelingMode,
+				LOCTEXT("StaticMeshEditorModelingMode", "Modeling Tools"),
+				LOCTEXT("StaticMeshEditorModelingModeTooltip", "Opens the Modeling Tools palette that provides selected mesh modification tools."),
+				FSlateIcon("ModelingToolsStyle", "LevelEditor.ModelingToolsMode")));
+		}));
+	}
+}
+
+
+#undef LOCTEXT_NAMESPACE
+	
+IMPLEMENT_MODULE(FStaticMeshEditorModelingModule, StaticMeshEditorModeling)
+
diff --git a/Engine/Plugins/Editor/StaticMeshEditorModeling/Source/StaticMeshEditorModeling/Private/StaticMeshEditorModelingToolkit.cpp b/Engine/Plugins/Editor/StaticMeshEditorModeling/Source/StaticMeshEditorModeling/Private/StaticMeshEditorModelingToolkit.cpp
new file mode 100644
--- /dev/null
+++ b/Engine/Plugins/Editor/StaticMeshEditorModeling/Source/StaticMeshEditorModeling/Private/StaticMeshEditorModelingToolkit.cpp
@@ -0,0 +1,293 @@
+// Copyright Epic Games, Inc. All Rights Reserved.
+
+#include "StaticMeshEditorModelingToolkit.h"
+#include "Widgets/Images/SImage.h"
+#include "Widgets/Input/SButton.h"
+#include "ModelingToolsEditorModeStyle.h"
+#include "ModelingToolsManagerActions.h"
+#include "Tools/UEdMode.h"
+#include "InteractiveTool.h"
+#include "EdModeInteractiveToolsContext.h"
+#include "InteractiveToolManager.h"
+#include "Widgets/Text/STextBlock.h"
+#include "Styling/AppStyle.h"
+#include "IDetailsView.h"
+#include "Framework/MultiBox/MultiBoxBuilder.h"
+#include "Tools/GenerateStaticMeshLODAssetTool.h"
+
+#define LOCTEXT_NAMESPACE "StaticMeshEditorModelingToolkit"
+
+static const FName LODManagementTabName(TEXT("LODManagement"));
+
+void FStaticMeshEditorModelingToolkit::Init(
+	const TSharedPtr<IToolkitHost>& InToolkitHost,
+	TWeakObjectPtr<UEdMode> InOwningMode )
+{
+	ModeWarningArea = SNew(STextBlock)
+		.AutoWrapText(true)
+		.Font(FCoreStyle::GetDefaultFontStyle("Bold", 9))
+		.ColorAndOpacity(FSlateColor(FLinearColor(0.9f, 0.15f, 0.15f)));
+	ModeWarningArea->SetText(FText::GetEmpty());
+	ModeWarningArea->SetVisibility(EVisibility::Collapsed);
+
+	ModeHeaderArea = SNew(STextBlock)
+		.AutoWrapText(true)
+		.Font(FCoreStyle::GetDefaultFontStyle("Bold", 12));
+	ModeHeaderArea->SetText(LOCTEXT("SelectToolLabel", "Select a Tool from the Toolbar"));
+	ModeHeaderArea->SetJustification(ETextJustify::Center);
+
+	ToolWarningArea = SNew(STextBlock)
+		.AutoWrapText(true)
+		.Font(FCoreStyle::GetDefaultFontStyle("Bold", 9))
+		.ColorAndOpacity(FSlateColor(FLinearColor(0.9f, 0.15f, 0.15f)));
+	ToolWarningArea->SetText(FText::GetEmpty());
+
+	SAssignNew(ToolkitWidget, SBorder)
+		.HAlign(HAlign_Fill)
+		.Padding(4)
+		[
+			SNew(SVerticalBox)
+			+ SVerticalBox::Slot().AutoHeight().HAlign(HAlign_Fill).Padding(5)
+			[
+				ModeWarningArea->AsShared()
+			]
+			+ SVerticalBox::Slot().AutoHeight().HAlign(HAlign_Fill).Padding(5)
+			[
+				ModeHeaderArea->AsShared()
+			]
+			+ SVerticalBox::Slot().AutoHeight().HAlign(HAlign_Fill).Padding(5)
+			[
+				ToolWarningArea->AsShared()
+			]
+			+ SVerticalBox::Slot().FillHeight(1.f).HAlign(HAlign_Fill)
+			[
+				SAssignNew(ToolDetailsContainer, SBorder)
+				.BorderImage(FAppStyle::GetBrush("NoBorder"))
+			]
+		];
+
+
+	FModeToolkit::Init(InToolkitHost, InOwningMode);
+
+	ToolDetailsContainer->SetContent(DetailsView.ToSharedRef());
+
+	ClearNotification();
+	ClearWarning();
+
+	ActiveToolName = FText::GetEmpty();
+	ActiveToolMessage = FText::GetEmpty();
+
+	GetScriptableEditorMode()->GetInteractiveToolsContext()->OnToolNotificationMessage.AddSP(this, &FStaticMeshEditorModelingToolkit::PostNotification);
+	GetScriptableEditorMode()->GetInteractiveToolsContext()->OnToolWarningMessage.AddSP(this, &FStaticMeshEditorModelingToolkit::PostWarning);
+
+
+	// ViewportOverlay
+	SAssignNew(ViewportOverlayWidget, SHorizontalBox)
+
+	+SHorizontalBox::Slot()
+	.HAlign(HAlign_Center)
+	.VAlign(VAlign_Bottom)
+	.Padding(FMargin(0.0f, 0.0f, 0.f, 15.f))
+	[
+		SNew(SBorder)
+		.BorderImage(FAppStyle::Get().GetBrush("EditorViewport.OverlayBrush"))
+		.Padding(8.f)
+		[
+			SNew(SHorizontalBox)
+
+			+SHorizontalBox::Slot()
+			.AutoWidth()
+			.VAlign(VAlign_Center)
+			.Padding(FMargin(0.f, 0.f, 8.f, 0.f))
+			[
+				SNew(SImage)
+				.Image_Lambda([this] () { return ActiveToolIcon; })
+			]
+
+			+SHorizontalBox::Slot()
+			.AutoWidth()
+			.VAlign(VAlign_Center)
+			.Padding(FMargin(0.f, 0.f, 8.f, 0.f))
+			[
+				SNew(STextBlock)
+				.Text(this, &FStaticMeshEditorModelingToolkit::GetActiveToolDisplayName)
+			]
+
+			+SHorizontalBox::Slot()
+			.AutoWidth()
+			.Padding(FMargin(0.0, 0.f, 2.f, 0.f))
+			[
+				SNew(SButton)
+				.ButtonStyle(FAppStyle::Get(), "PrimaryButton")
+				.TextStyle( FAppStyle::Get(), "DialogButtonText" )
+				.Text(LOCTEXT("OverlayAccept", "Accept"))
+				.ToolTipText(LOCTEXT("OverlayAcceptTooltip", "Accept/Commit the results of the active Tool [Enter]"))
+				.HAlign(HAlign_Center)
+				.OnClicked_Lambda([this]() { GetScriptableEditorMode()->GetInteractiveToolsContext()->EndTool(EToolShutdownType::Accept); return FReply::Handled(); })
+				.IsEnabled_Lambda([this]() { return GetScriptableEditorMode()->GetInteractiveToolsContext()->CanAcceptActiveTool(); })
+				.Visibility_Lambda([this]() { return GetScriptableEditorMode()->GetInteractiveToolsContext()->ActiveToolHasAccept() ? EVisibility::Visible : EVisibility::Collapsed; })
+			]
+
+			+SHorizontalBox::Slot()
+			.AutoWidth()
+			.Padding(FMargin(2.0, 0.f, 0.f, 0.f))
+			[
+				SNew(SButton)
+				.ButtonStyle(FAppStyle::Get(), "Button")
+				.TextStyle( FAppStyle::Get(), "DialogButtonText" )
+				.Text(LOCTEXT("OverlayCancel", "Cancel"))
+				.ToolTipText(LOCTEXT("OverlayCancelTooltip", "Cancel the active Tool [Esc]"))
+				.HAlign(HAlign_Center)
+				.OnClicked_Lambda([this]() { GetScriptableEditorMode()->GetInteractiveToolsContext()->EndTool(EToolShutdownType::Cancel); return FReply::Handled(); })
+				.IsEnabled_Lambda([this]() { return GetScriptableEditorMode()->GetInteractiveToolsContext()->CanCancelActiveTool(); })
+				.Visibility_Lambda([this]() { return GetScriptableEditorMode()->GetInteractiveToolsContext()->ActiveToolHasAccept() ? EVisibility::Visible : EVisibility::Collapsed; })
+			]
+
+			+SHorizontalBox::Slot()
+			.AutoWidth()
+			.Padding(FMargin(2.0, 0.f, 0.f, 0.f))
+			[
+				SNew(SButton)
+				.ButtonStyle(FAppStyle::Get(), "PrimaryButton")
+				.TextStyle( FAppStyle::Get(), "DialogButtonText" )
+				.Text(LOCTEXT("OverlayComplete", "Complete"))
+				.ToolTipText(LOCTEXT("OverlayCompleteTooltip", "Exit the active Tool [Enter]"))
+				.HAlign(HAlign_Center)
+				.OnClicked_Lambda([this]() { GetScriptableEditorMode()->GetInteractiveToolsContext()->EndTool(EToolShutdownType::Completed); return FReply::Handled(); })
+				.IsEnabled_Lambda([this]() { return GetScriptableEditorMode()->GetInteractiveToolsContext()->CanCompleteActiveTool(); })
+				.Visibility_Lambda([this]() { return GetScriptableEditorMode()->GetInteractiveToolsContext()->CanCompleteActiveTool() ? EVisibility::Visible : EVisibility::Collapsed; })
+			]
+		]	
+	];
+
+	
+	SetCurrentPalette("LODManagement");
+}
+
+
+TSharedPtr<SWidget> FStaticMeshEditorModelingToolkit::GetInlineContent() const
+{
+	return ToolkitWidget;
+}
+
+FName FStaticMeshEditorModelingToolkit::GetToolkitFName() const
+{
+	return FName("MeshLODEditModeToolkit");
+}
+
+FText FStaticMeshEditorModelingToolkit::GetBaseToolkitName() const
+{
+	return LOCTEXT("DisplayName", "FMeshLODPluginEditModeToolkit Tool");
+}
+
+
+void FStaticMeshEditorModelingToolkit::GetToolPaletteNames(TArray<FName>& PaletteNames) const
+{
+	PaletteNames = { LODManagementTabName };
+}
+
+FText FStaticMeshEditorModelingToolkit::GetToolPaletteDisplayName(FName PaletteName) const
+{
+	return FText::FromName(PaletteName);
+}
+
+void FStaticMeshEditorModelingToolkit::BuildToolPalette(FName PaletteName, class FToolBarBuilder& ToolbarBuilder)
+{
+	const FModelingToolsManagerCommands& Commands = FModelingToolsManagerCommands::Get();
+	
+	if (PaletteName == LODManagementTabName)
+	{
+		ToolbarBuilder.AddToolBarButton(Commands.BeginGenerateStaticMeshLODAssetTool);
+		ToolbarBuilder.AddToolBarButton(Commands.BeginLODManagerTool);
+		ToolbarBuilder.AddToolBarButton(Commands.BeginMeshInspectorTool);
+	}
+}
+
+void FStaticMeshEditorModelingToolkit::PostNotification(const FText& InMessage)
+{
+	ActiveToolMessage = InMessage;
+}
+
+
+void FStaticMeshEditorModelingToolkit::ClearNotification()
+{
+	ActiveToolMessage = FText::GetEmpty();
+}
+
+FText FStaticMeshEditorModelingToolkit::GetActiveToolDisplayName() const
+{
+	return ActiveToolName;
+}
+
+FText FStaticMeshEditorModelingToolkit::GetActiveToolMessage() const
+{
+	return ActiveToolMessage;
+}
+
+void FStaticMeshEditorModelingToolkit::PostWarning(const FText& Message)
+{
+	ToolWarningArea->SetText(Message);
+	ToolWarningArea->SetVisibility(EVisibility::Visible);
+}
+
+
+void FStaticMeshEditorModelingToolkit::ClearWarning()
+{
+	ToolWarningArea->SetText(FText());
+	ToolWarningArea->SetVisibility(EVisibility::Collapsed);
+}
+
+void FStaticMeshEditorModelingToolkit::UpdateActiveToolProperties(UInteractiveTool* Tool)
+{
+	if (Tool)
+	{
+		DetailsView->SetObjects(Tool->GetToolProperties(true));
+	}
+}
+
+void FStaticMeshEditorModelingToolkit::OnToolStarted(UInteractiveToolManager* Manager, UInteractiveTool* Tool)
+{
+	FModeToolkit::OnToolStarted(Manager, Tool);
+
+	Tool->OnPropertySetsModified.AddSP(this, &FStaticMeshEditorModelingToolkit::UpdateActiveToolProperties, Tool);
+
+	ModeHeaderArea->SetVisibility(EVisibility::Collapsed);
+	ActiveToolName = Tool->GetToolInfo().ToolDisplayName;
+
+	FString ActiveToolIdentifier = GetScriptableEditorMode()->GetToolManager()->GetActiveToolName(EToolSide::Left);
+	ActiveToolIdentifier.InsertAt(0, ".");
+	FName ActiveToolIconName = ISlateStyle::Join(FModelingToolsManagerCommands::Get().GetContextName(), TCHAR_TO_ANSI(*ActiveToolIdentifier));
+	ActiveToolIcon = FModelingToolsEditorModeStyle::Get()->GetOptionalBrush(ActiveToolIconName);
+
+	GetToolkitHost()->AddViewportOverlayWidget(ViewportOverlayWidget.ToSharedRef());
+
+	const UInteractiveTool* CurrentTool = GetScriptableEditorMode()->GetToolManager()->GetActiveTool(EToolSide::Left);
+	if (CurrentTool->GetClass()->IsChildOf<UGenerateStaticMeshLODAssetTool>())
+	{
+		PostWarning(LOCTEXT("AutoLODStaticMeshEditorWarning", "Note: This tool may create new assets including new Static Meshes, Textures, and Materials"));
+	}
+}
+
+
+void FStaticMeshEditorModelingToolkit::OnToolEnded(UInteractiveToolManager* Manager, UInteractiveTool* Tool)
+{
+	FModeToolkit::OnToolEnded(Manager, Tool);
+
+	if (IsHosted())
+	{
+		GetToolkitHost()->RemoveViewportOverlayWidget(ViewportOverlayWidget.ToSharedRef());
+	}
+
+	if (Tool)
+	{
+		Tool->OnPropertySetsModified.RemoveAll(this);
+	}
+
+	ModeHeaderArea->SetVisibility(EVisibility::Visible);
+	ActiveToolName = FText::GetEmpty();
+	ActiveToolIcon = nullptr;
+	ClearNotification();
+	ClearWarning();
+}
+
+#undef LOCTEXT_NAMESPACE
diff --git a/Engine/Plugins/Editor/StaticMeshEditorModeling/Source/StaticMeshEditorModeling/Public/StaticMeshEditorModelingMode.h b/Engine/Plugins/Editor/StaticMeshEditorModeling/Source/StaticMeshEditorModeling/Public/StaticMeshEditorModelingMode.h
new file mode 100644
--- /dev/null
+++ b/Engine/Plugins/Editor/StaticMeshEditorModeling/Source/StaticMeshEditorModeling/Public/StaticMeshEditorModelingMode.h
@@ -0,0 +1,31 @@
+// Copyright Epic Games, Inc. All Rights Reserved.
+
+#pragma once
+
+#include "Tools/LegacyEdModeWidgetHelpers.h"
+#include "Tools/Modes.h"
+#include "StaticMeshEditorModelingMode.generated.h"
+
+class SBorder;
+
+UCLASS()
+class UStaticMeshEditorModelingMode : public UBaseLegacyWidgetEdMode
+{
+	GENERATED_BODY()
+
+public:
+
+	const static FEditorModeID Id;
+
+protected:
+
+	UStaticMeshEditorModelingMode();
+
+	virtual void Enter() override;
+	virtual bool OnRequestClose() override;
+
+	virtual bool UsesToolkits() const override;
+	virtual void CreateToolkit() override;
+	
+
+};
diff --git a/Engine/Plugins/Editor/StaticMeshEditorModeling/Source/StaticMeshEditorModeling/Public/StaticMeshEditorModelingModule.h b/Engine/Plugins/Editor/StaticMeshEditorModeling/Source/StaticMeshEditorModeling/Public/StaticMeshEditorModelingModule.h
new file mode 100644
--- /dev/null
+++ b/Engine/Plugins/Editor/StaticMeshEditorModeling/Source/StaticMeshEditorModeling/Public/StaticMeshEditorModelingModule.h
@@ -0,0 +1,29 @@
+// Copyright Epic Games, Inc. All Rights Reserved.
+
+#pragma once
+
+#include "Modules/ModuleInterface.h"
+#include "Templates/SharedPointer.h"
+
+class IStaticMeshEditor;
+class SDockTab;
+
+class FStaticMeshEditorModelingModule : public IModuleInterface
+{
+public:
+
+	/** IModuleInterface implementation */
+	virtual void StartupModule() override;
+	virtual void ShutdownModule() override;
+
+	bool IsStaticMeshEditorModelingModeActive(TWeakPtr<IStaticMeshEditor> InEditor) const;
+	void OnToggleStaticMeshEditorModelingMode(TWeakPtr<IStaticMeshEditor> InEditor);
+
+	void RegisterMenusAndToolbars();
+
+private:
+
+	TWeakPtr<SDockTab> MeshLODDockTab;
+
+};
+
diff --git a/Engine/Plugins/Editor/StaticMeshEditorModeling/Source/StaticMeshEditorModeling/Public/StaticMeshEditorModelingToolkit.h b/Engine/Plugins/Editor/StaticMeshEditorModeling/Source/StaticMeshEditorModeling/Public/StaticMeshEditorModelingToolkit.h
new file mode 100644
--- /dev/null
+++ b/Engine/Plugins/Editor/StaticMeshEditorModeling/Source/StaticMeshEditorModeling/Public/StaticMeshEditorModelingToolkit.h
@@ -0,0 +1,52 @@
+// Copyright Epic Games, Inc. All Rights Reserved.
+
+#pragma once
+
+#include "Toolkits/BaseToolkit.h"
+
+class STextBlock;
+
+class FStaticMeshEditorModelingToolkit : public FModeToolkit
+{
+public:
+
+	void Init( const TSharedPtr<IToolkitHost>& InToolkitHost, TWeakObjectPtr<UEdMode> InOwningMode ) override;
+
+	FName GetToolkitFName() const override;
+	FText GetBaseToolkitName() const override;
+	TSharedPtr<SWidget> GetInlineContent() const override;
+
+	void GetToolPaletteNames(TArray<FName>& InPaletteName) const override;
+	FText GetToolPaletteDisplayName(FName PaletteName) const override;
+	void BuildToolPalette(FName PaletteName, class FToolBarBuilder& ToolbarBuilder) override;
+
+	void OnToolStarted(UInteractiveToolManager* Manager, UInteractiveTool* Tool) override;
+	void OnToolEnded(UInteractiveToolManager* Manager, UInteractiveTool* Tool) override;
+
+private:
+
+	void PostNotification(const FText& Message);
+	void ClearNotification();
+	void PostWarning(const FText& Message);
+	void ClearWarning();
+
+	void UpdateActiveToolProperties(UInteractiveTool* Tool);
+
+	FText GetActiveToolDisplayName() const override;
+	FText GetActiveToolMessage() const override;
+
+	TSharedPtr<SBorder> ToolDetailsContainer;
+
+	TSharedPtr<SWidget> ViewportOverlayWidget;
+	const FSlateBrush* ActiveToolIcon = nullptr;
+	FText ActiveToolName;
+	FText ActiveToolMessage;
+
+	TSharedPtr<STextBlock> ModeWarningArea;
+	TSharedPtr<STextBlock> ModeHeaderArea;
+	TSharedPtr<STextBlock> ToolWarningArea;
+
+};
+
+
+
diff --git a/Engine/Plugins/Editor/StaticMeshEditorModeling/Source/StaticMeshEditorModeling/StaticMeshEditorModeling.Build.cs b/Engine/Plugins/Editor/StaticMeshEditorModeling/Source/StaticMeshEditorModeling/StaticMeshEditorModeling.Build.cs
new file mode 100644
--- /dev/null
+++ b/Engine/Plugins/Editor/StaticMeshEditorModeling/Source/StaticMeshEditorModeling/StaticMeshEditorModeling.Build.cs
@@ -0,0 +1,63 @@
+// Copyright Epic Games, Inc. All Rights Reserved.
+
+using UnrealBuildTool;
+
+public class StaticMeshEditorModeling : ModuleRules
+{
+	public StaticMeshEditorModeling(ReadOnlyTargetRules Target) : base(Target)
+	{
+		PCHUsage = ModuleRules.PCHUsageMode.UseExplicitOrSharedPCHs;
+		
+		PublicIncludePaths.AddRange(
+			new string[] {
+				// ... add public include paths required here ...
+			}
+			);
+				
+		
+		PrivateIncludePaths.AddRange(
+			new string[] {
+				// ... add other private include paths required here ...
+			}
+			);
+			
+		
+		PublicDependencyModuleNames.AddRange(
+			new string[]
+			{			
+				// ... add other public dependencies that you statically link with here ...
+			}
+			);
+			
+		
+		PrivateDependencyModuleNames.AddRange(
+			new string[]
+			{
+				"Core",
+				"InputCore",
+				"CoreUObject",
+				"EditorInteractiveToolsFramework",				
+				"EditorFramework",
+				"InteractiveToolsFramework",
+				"MeshLODToolset",
+				"ModelingComponentsEditorOnly",
+				"ModelingToolsEditorMode",
+				"Slate",
+				"SlateCore",
+				"StaticMeshEditor",
+				"MeshModelingToolsExp",
+				"ToolMenus",
+				"UnrealEd",
+			}
+			);
+		
+		
+		DynamicallyLoadedModuleNames.AddRange(
+			new string[]
+			{
+				// ... add any modules that your module loads dynamically here ...
+			}
+			);
+
+	}
+}
diff --git a/Engine/Plugins/Editor/StaticMeshEditorModeling/StaticMeshEditorModeling.uplugin b/Engine/Plugins/Editor/StaticMeshEditorModeling/StaticMeshEditorModeling.uplugin
new file mode 100644
--- /dev/null
+++ b/Engine/Plugins/Editor/StaticMeshEditorModeling/StaticMeshEditorModeling.uplugin
@@ -0,0 +1,35 @@
+{
+	"FileVersion": 3,
+	"Version": 1,
+	"VersionName": "0.1",
+	"FriendlyName": "Static Mesh Editor Modeling Mode",
+	"Description": "Enable a Modeling Tools Tab in the Static Mesh Editor",
+	"Category": "Editor",
+	"CreatedBy": "Epic Games, Inc.",
+	"CreatedByURL": "https://epicgames.com",
+	"DocsURL": "",
+	"MarketplaceURL": "",
+	"SupportURL": "",
+	"CanContainContent": true,
+	"IsBetaVersion": true,
+	"Installed": false,
+	"Hidden": false,
+	"Modules": [
+		{
+			"Name": "StaticMeshEditorModeling",
+			"Type": "Editor",
+			"LoadingPhase": "Default",
+			"PlatformAllowList": [ "Win64" ]
+		},
+	],
+	"Plugins": [
+		{
+			"Name": "MeshLODToolset",
+			"Enabled": true
+		},
+		{
+			"Name": "ModelingToolsEditorMode",
+			"Enabled": true
+		}
+	]
+}
diff --git a/Engine/Plugins/Experimental/StaticMeshEditorModeling/Source/StaticMeshEditorModeling/Private/StaticMeshEditorModelingCommands.cpp b/Engine/Plugins/Experimental/StaticMeshEditorModeling/Source/StaticMeshEditorModeling/Private/StaticMeshEditorModelingCommands.cpp
deleted file mode 100644
diff --git a/Engine/Plugins/Experimental/StaticMeshEditorModeling/Source/StaticMeshEditorModeling/Private/StaticMeshEditorModelingCommands.h b/Engine/Plugins/Experimental/StaticMeshEditorModeling/Source/StaticMeshEditorModeling/Private/StaticMeshEditorModelingCommands.h
deleted file mode 100644
diff --git a/Engine/Plugins/Experimental/StaticMeshEditorModeling/Source/StaticMeshEditorModeling/Private/StaticMeshEditorModelingMode.cpp b/Engine/Plugins/Experimental/StaticMeshEditorModeling/Source/StaticMeshEditorModeling/Private/StaticMeshEditorModelingMode.cpp
deleted file mode 100644
diff --git a/Engine/Plugins/Experimental/StaticMeshEditorModeling/Source/StaticMeshEditorModeling/Private/StaticMeshEditorModelingModule.cpp b/Engine/Plugins/Experimental/StaticMeshEditorModeling/Source/StaticMeshEditorModeling/Private/StaticMeshEditorModelingModule.cpp
deleted file mode 100644
diff --git a/Engine/Plugins/Experimental/StaticMeshEditorModeling/Source/StaticMeshEditorModeling/Private/StaticMeshEditorModelingToolkit.cpp b/Engine/Plugins/Experimental/StaticMeshEditorModeling/Source/StaticMeshEditorModeling/Private/StaticMeshEditorModelingToolkit.cpp
deleted file mode 100644
diff --git a/Engine/Plugins/Experimental/StaticMeshEditorModeling/Source/StaticMeshEditorModeling/Public/StaticMeshEditorModelingMode.h b/Engine/Plugins/Experimental/StaticMeshEditorModeling/Source/StaticMeshEditorModeling/Public/StaticMeshEditorModelingMode.h
deleted file mode 100644
diff --git a/Engine/Plugins/Experimental/StaticMeshEditorModeling/Source/StaticMeshEditorModeling/Public/StaticMeshEditorModelingModule.h b/Engine/Plugins/Experimental/StaticMeshEditorModeling/Source/StaticMeshEditorModeling/Public/StaticMeshEditorModelingModule.h
deleted file mode 100644
diff --git a/Engine/Plugins/Experimental/StaticMeshEditorModeling/Source/StaticMeshEditorModeling/Public/StaticMeshEditorModelingToolkit.h b/Engine/Plugins/Experimental/StaticMeshEditorModeling/Source/StaticMeshEditorModeling/Public/StaticMeshEditorModelingToolkit.h
deleted file mode 100644
diff --git a/Engine/Plugins/Experimental/StaticMeshEditorModeling/Source/StaticMeshEditorModeling/StaticMeshEditorModeling.Build.cs b/Engine/Plugins/Experimental/StaticMeshEditorModeling/Source/StaticMeshEditorModeling/StaticMeshEditorModeling.Build.cs
deleted file mode 100644
diff --git a/Engine/Plugins/Experimental/StaticMeshEditorModeling/StaticMeshEditorModeling.uplugin b/Engine/Plugins/Experimental/StaticMeshEditorModeling/StaticMeshEditorModeling.uplugin
deleted file mode 100644
